READY TO SUBMIT

diff --git a/Engine/Plugins/Experimental/LearningAgents/Source/LearningAgents/Public/LearningAgentsObservations.h b/Engine/Plugins/Experimental/LearningAgents/Source/LearningAgents/Public/LearningAgentsObservations.h
--- a/Engine/Plugins/Experimental/LearningAgents/Source/LearningAgents/Public/LearningAgentsObservations.h
+++ b/Engine/Plugins/Experimental/LearningAgents/Source/LearningAgents/Public/LearningAgentsObservations.h
@@ -297,96 +297,105 @@ public:
 	 * @param Schema The Observation Schema
 	 * @param Size The number of float values in the observation.
 	 * @param Scale The scale used to normalize this observation.
+	 * @param bActive Whether to enable the observation as a valid network input. If set to inactive, this observation will be considered an empty observation. This flag allows easy toggles of specified observation.
 	 * @param Tag The tag of this new observation. Used during observation object validation and debugging.
 	 * @return The newly created observation schema element.
 	 */
 	UFUNCTION(BlueprintPure, Category = "LearningAgents", meta = (AdvancedDisplay = 3))
-	static UE_API FLearningAgentsObservationSchemaElement SpecifyContinuousObservation(ULearningAgentsObservationSchema* Schema, const int32 Size, const float Scale = 1.0f, const FName Tag = TEXT("ContinuousObservation"));
+	static UE_API FLearningAgentsObservationSchemaElement SpecifyContinuousObservation(ULearningAgentsObservationSchema* Schema, const int32 Size, const float Scale = 1.0f, const bool bActive = true, const FName Tag = TEXT("ContinuousObservation"));
 
 	/**
 	 * Specifies a new exclusive discrete observation. This represents a discrete observation which is an exclusive selection from multiple choices.
 	 *
 	 * @param Schema The Observation Schema
 	 * @param Size The number of discrete options in the observation.
+	 * @param bActive Whether to enable the observation as a valid network input. If set to inactive, this observation will be considered an empty observation. This flag allows easy toggles of specified observation.
 	 * @param Tag The tag of this new observation. Used during observation object validation and debugging.
 	 * @return The newly created observation schema element.
 	 */
 	UFUNCTION(BlueprintPure, Category = "LearningAgents", meta = (AdvancedDisplay = 2))
-	static UE_API FLearningAgentsObservationSchemaElement SpecifyExclusiveDiscreteObservation(ULearningAgentsObservationSchema* Schema, const int32 Size, const FName Tag = TEXT("ExclusiveDiscreteObservation"));
+	static UE_API FLearningAgentsObservationSchemaElement SpecifyExclusiveDiscreteObservation(ULearningAgentsObservationSchema* Schema, const int32 Size, const bool bActive = true, const FName Tag = TEXT("ExclusiveDiscreteObservation"));
 
 	/**
 	 * Specifies a new inclusive discrete observation. This represents a discrete observation which is an inclusive selection from multiple choices.
 	 *
 	 * @param Schema The Observation Schema
 	 * @param Size The number of discrete options in the observation.
+	 * @param bActive Whether to enable the observation as a valid network input. If set to inactive, this observation will be considered an empty observation. This flag allows easy toggles of specified observation.
 	 * @param Tag The tag of this new observation. Used during observation object validation and debugging.
 	 * @return The newly created observation schema element.
 	 */
 	UFUNCTION(BlueprintPure, Category = "LearningAgents", meta = (AdvancedDisplay = 2))
-	static UE_API FLearningAgentsObservationSchemaElement SpecifyInclusiveDiscreteObservation(ULearningAgentsObservationSchema* Schema, const int32 Size, const FName Tag = TEXT("InclusiveDiscreteObservation"));
+	static UE_API FLearningAgentsObservationSchemaElement SpecifyInclusiveDiscreteObservation(ULearningAgentsObservationSchema* Schema, const int32 Size, const bool bActive = true, const FName Tag = TEXT("InclusiveDiscreteObservation"));
 
 	/**
 	 * Specifies a new named exclusive discrete observation. This represents a discrete observation which is an exclusive selection from multiple choices.
 	 *
 	 * @param Schema The Observation Schema
 	 * @param ElementNames The names of the discrete options in the observation.
+	 * @param bActive Whether to enable the observation as a valid network input. If set to inactive, this observation will be considered an empty observation. This flag allows easy toggles of specified observation.
 	 * @param Tag The tag of this new observation. Used during observation object validation and debugging.
 	 * @return The newly created observation schema element.
 	 */
 	UFUNCTION(BlueprintPure, Category = "LearningAgents", meta = (AdvancedDisplay = 2))
-	static UE_API FLearningAgentsObservationSchemaElement SpecifyNamedExclusiveDiscreteObservation(ULearningAgentsObservationSchema* Schema, const TArray<FName>& ElementNames, const FName Tag = TEXT("NamedExclusiveDiscreteObservation"));
+	static UE_API FLearningAgentsObservationSchemaElement SpecifyNamedExclusiveDiscreteObservation(ULearningAgentsObservationSchema* Schema, const TArray<FName>& ElementNames, const bool bActive = true, const FName Tag = TEXT("NamedExclusiveDiscreteObservation"));
 
 	/**
 	 * Specifies a new named exclusive discrete observation. This represents a discrete observation which is an exclusive selection from multiple choices.
 	 *
 	 * @param Schema The Observation Schema
 	 * @param ElementNames The names of the discrete options in the observation.
+	 * @param bActive Whether to enable the observation as a valid network input. If set to inactive, this observation will be considered an empty observation. This flag allows easy toggles of specified observation.
 	 * @param Tag The tag of this new observation. Used during observation object validation and debugging.
 	 * @return The newly created observation schema element.
 	 */
-	static UE_API FLearningAgentsObservationSchemaElement SpecifyNamedExclusiveDiscreteObservationFromArrayView(ULearningAgentsObservationSchema* Schema, const TArrayView<const FName> ElementNames, const FName Tag = TEXT("NamedExclusiveDiscreteObservation"));
+	static UE_API FLearningAgentsObservationSchemaElement SpecifyNamedExclusiveDiscreteObservationFromArrayView(ULearningAgentsObservationSchema* Schema, const TArrayView<const FName> ElementNames, const bool bActive = true, const FName Tag = TEXT("NamedExclusiveDiscreteObservation"));
 
 	/**
 	 * Specifies a new named inclusive discrete observation. This represents a discrete observation which is an inclusive selection from multiple choices.
 	 *
 	 * @param Schema The Observation Schema
 	 * @param ElementNames The names of the discrete options in the observation.
+	 * @param bActive Whether to enable the observation as a valid network input. If set to inactive, this observation will be considered an empty observation. This flag allows easy toggles of specified observation.
 	 * @param Tag The tag of this new observation. Used during observation object validation and debugging.
 	 * @return The newly created observation schema element.
 	 */
 	UFUNCTION(BlueprintPure, Category = "LearningAgents", meta = (AdvancedDisplay = 2))
-	static UE_API FLearningAgentsObservationSchemaElement SpecifyNamedInclusiveDiscreteObservation(ULearningAgentsObservationSchema* Schema, const TArray<FName>& ElementNames, const FName Tag = TEXT("NamedInclusiveDiscreteObservation"));
+	static UE_API FLearningAgentsObservationSchemaElement SpecifyNamedInclusiveDiscreteObservation(ULearningAgentsObservationSchema* Schema, const TArray<FName>& ElementNames, const bool bActive = true, const FName Tag = TEXT("NamedInclusiveDiscreteObservation"));
 
 	/**
 	 * Specifies a new named inclusive discrete observation. This represents a discrete observation which is an inclusive selection from multiple choices.
 	 *
 	 * @param Schema The Observation Schema
 	 * @param ElementNames The names of the discrete options in the observation.
+	 * @param bActive Whether to enable the observation as a valid network input. If set to inactive, this observation will be considered an empty observation. This flag allows easy toggles of specified observation.
 	 * @param Tag The tag of this new observation. Used during observation object validation and debugging.
 	 * @return The newly created observation schema element.
 	 */
-	static UE_API FLearningAgentsObservationSchemaElement SpecifyNamedInclusiveDiscreteObservationFromArrayView(ULearningAgentsObservationSchema* Schema, const TArrayView<const FName> ElementNames, const FName Tag = TEXT("NamedInclusiveDiscreteObservation"));
+	static UE_API FLearningAgentsObservationSchemaElement SpecifyNamedInclusiveDiscreteObservationFromArrayView(ULearningAgentsObservationSchema* Schema, const TArrayView<const FName> ElementNames, const bool bActive = true, const FName Tag = TEXT("NamedInclusiveDiscreteObservation"));
 
 	/**
 	 * Specifies a new count observation. This represents a count of something such as the size of, or index into, an array.
 	 *
 	 * @param Schema The Observation Schema
+	 * @param bActive Whether to enable the observation as a valid network input. If set to inactive, this observation will be considered an empty observation. This flag allows easy toggles of specified observation.
 	 * @param Tag The tag of this new observation. Used during observation object validation and debugging.
 	 * @return The newly created observation schema element.
 	 */
 	UFUNCTION(BlueprintPure, Category = "LearningAgents", meta = (AdvancedDisplay = 1))
-	static UE_API FLearningAgentsObservationSchemaElement SpecifyCountObservation(ULearningAgentsObservationSchema* Schema, const FName Tag = TEXT("CountObservation"));
+	static UE_API FLearningAgentsObservationSchemaElement SpecifyCountObservation(ULearningAgentsObservationSchema* Schema, const bool bActive = true, const FName Tag = TEXT("CountObservation"));
 
 	/**
 	 * Specifies a new struct observation. This represents a group of named sub-observations.
 	 *
 	 * @param Schema The Observation Schema
 	 * @param Elements The sub-observations that make up this struct.
+	 * @param bActive Whether to enable the observation as a valid network input. If set to inactive, this observation will be considered an empty observation. This flag allows easy toggles of specified observation.
 	 * @param Tag The tag of this new observation. Used during observation object validation and debugging.
 	 * @return The newly created observation schema element.
 	 */
 	UFUNCTION(BlueprintPure, Category = "LearningAgents", meta = (AdvancedDisplay = 2))
-	static UE_API FLearningAgentsObservationSchemaElement SpecifyStructObservation(ULearningAgentsObservationSchema* Schema, const TMap<FName, FLearningAgentsObservationSchemaElement>& Elements, const FName Tag = TEXT("StructObservation"));
+	static UE_API FLearningAgentsObservationSchemaElement SpecifyStructObservation(ULearningAgentsObservationSchema* Schema, const TMap<FName, FLearningAgentsObservationSchemaElement>& Elements, const bool bActive = true, const FName Tag = TEXT("StructObservation"));
 
 	/**
 	 * Specifies a new struct observation. This represents a group of named sub-observations.
@@ -394,11 +403,12 @@ public:
 	 * @param Schema The Observation Schema
 	 * @param ElementNames The names of the sub-observations that make up this struct.
 	 * @param Elements The corresponding sub-observations that make up this struct. Must be the same size as ElementNames.
+	 * @param bActive Whether to enable the observation as a valid network input. If set to inactive, this observation will be considered an empty observation. This flag allows easy toggles of specified observation.
 	 * @param Tag The tag of this new observation. Used during observation object validation and debugging.
 	 * @return The newly created observation schema element.
 	 */
 	UFUNCTION(BlueprintPure, Category = "LearningAgents", meta = (AdvancedDisplay = 3))
-	static UE_API FLearningAgentsObservationSchemaElement SpecifyStructObservationFromArrays(ULearningAgentsObservationSchema* Schema, const TArray<FName>& ElementNames, const TArray<FLearningAgentsObservationSchemaElement>& Elements, const FName Tag = TEXT("StructObservation"));
+	static UE_API FLearningAgentsObservationSchemaElement SpecifyStructObservationFromArrays(ULearningAgentsObservationSchema* Schema, const TArray<FName>& ElementNames, const TArray<FLearningAgentsObservationSchemaElement>& Elements, const bool bActive = true, const FName Tag = TEXT("StructObservation"));
 	
 	/**
 	 * Specifies a new struct observation. This represents a group of named sub-observations.
@@ -406,10 +416,11 @@ public:
 	 * @param Schema The Observation Schema
 	 * @param ElementNames The names of the sub-observations that make up this struct.
 	 * @param Elements The corresponding sub-observations that make up this struct. Must be the same size as ElementNames.
+	 * @param bActive Whether to enable the observation as a valid network input. If set to inactive, this observation will be considered an empty observation. This flag allows easy toggles of specified observation.
 	 * @param Tag The tag of this new observation. Used during observation object validation and debugging.
 	 * @return The newly created observation schema element.
 	 */
-	static UE_API FLearningAgentsObservationSchemaElement SpecifyStructObservationFromArrayViews(ULearningAgentsObservationSchema* Schema, const TArrayView<const FName> ElementNames, const TArrayView<const FLearningAgentsObservationSchemaElement> Elements, const FName Tag = TEXT("StructObservation"));
+	static UE_API FLearningAgentsObservationSchemaElement SpecifyStructObservationFromArrayViews(ULearningAgentsObservationSchema* Schema, const TArrayView<const FName> ElementNames, const TArrayView<const FLearningAgentsObservationSchemaElement> Elements, const bool bActive = true, const FName Tag = TEXT("StructObservation"));
 
 	/**
 	 * Specifies a new exclusive union observation. This represents an observation which is exclusively chosen from a set of named sub-observations.
@@ -418,11 +429,12 @@ public:
 	 * @param Schema The Observation Schema
 	 * @param Elements The sub-observations that make up this union.
 	 * @param EncodingSize The encoding size used to encode each sub-observation.
+	 * @param bActive Whether to enable the observation as a valid network input. If set to inactive, this observation will be considered an empty observation. This flag allows easy toggles of specified observation.
 	 * @param Tag The tag of this new observation. Used during observation object validation and debugging.
 	 * @return The newly created observation schema element.
 	 */
 	UFUNCTION(BlueprintPure, Category = "LearningAgents", meta = (AdvancedDisplay = 2))
-	static UE_API FLearningAgentsObservationSchemaElement SpecifyExclusiveUnionObservation(ULearningAgentsObservationSchema* Schema, const TMap<FName, FLearningAgentsObservationSchemaElement>& Elements, const int32 EncodingSize = 128, const FName Tag = TEXT("ExclusiveUnionObservation"));
+	static UE_API FLearningAgentsObservationSchemaElement SpecifyExclusiveUnionObservation(ULearningAgentsObservationSchema* Schema, const TMap<FName, FLearningAgentsObservationSchemaElement>& Elements, const int32 EncodingSize = 128, const bool bActive = true, const FName Tag = TEXT("ExclusiveUnionObservation"));
 
 	/**
 	 * Specifies a new exclusive union observation. This represents an observation which is exclusively chosen from a set of named sub-observations.
@@ -432,11 +444,12 @@ public:
 	 * @param ElementNames The names of the sub-observations that make up this union.
 	 * @param Elements The corresponding sub-observations that make up this union. Must be the same size as ElementNames.
 	 * @param EncodingSize The encoding size used to encode each sub-observation.
+	 * @param bActive Whether to enable the observation as a valid network input. If set to inactive, this observation will be considered an empty observation. This flag allows easy toggles of specified observation.
 	 * @param Tag The tag of this new observation. Used during observation object validation and debugging.
 	 * @return The newly created observation schema element.
 	 */
 	UFUNCTION(BlueprintPure, Category = "LearningAgents", meta = (AdvancedDisplay = 3))
-	static UE_API FLearningAgentsObservationSchemaElement SpecifyExclusiveUnionObservationFromArrays(ULearningAgentsObservationSchema* Schema, const TArray<FName>& ElementNames, const TArray<FLearningAgentsObservationSchemaElement>& Elements, const int32 EncodingSize = 128, const FName Tag = TEXT("ExclusiveUnionObservation"));
+	static UE_API FLearningAgentsObservationSchemaElement SpecifyExclusiveUnionObservationFromArrays(ULearningAgentsObservationSchema* Schema, const TArray<FName>& ElementNames, const TArray<FLearningAgentsObservationSchemaElement>& Elements, const int32 EncodingSize = 128, const bool bActive = true, const FName Tag = TEXT("ExclusiveUnionObservation"));
 	
 	/**
 	 * Specifies a new exclusive union observation. This represents an observation which is exclusively chosen from a set of named sub-observations.
@@ -446,10 +459,11 @@ public:
 	 * @param ElementNames The names of the sub-observations that make up this union.
 	 * @param Elements The corresponding sub-observations that make up this union. Must be the same size as ElementNames.
 	 * @param EncodingSize The encoding size used to encode each sub-observation.
+	 * @param bActive Whether to enable the observation as a valid network input. If set to inactive, this observation will be considered an empty observation. This flag allows easy toggles of specified observation.
 	 * @param Tag The tag of this new observation. Used during observation object validation and debugging.
 	 * @return The newly created observation schema element.
 	 */
-	static UE_API FLearningAgentsObservationSchemaElement SpecifyExclusiveUnionObservationFromArrayViews(ULearningAgentsObservationSchema* Schema, const TArrayView<const FName> ElementNames, const TArrayView<const FLearningAgentsObservationSchemaElement> Elements, const int32 EncodingSize = 128, const FName Tag = TEXT("ExclusiveUnionObservation"));
+	static UE_API FLearningAgentsObservationSchemaElement SpecifyExclusiveUnionObservationFromArrayViews(ULearningAgentsObservationSchema* Schema, const TArrayView<const FName> ElementNames, const TArrayView<const FLearningAgentsObservationSchemaElement> Elements, const int32 EncodingSize = 128, const bool bActive = true, const FName Tag = TEXT("ExclusiveUnionObservation"));
 
 	/**
 	 * Specifies a new inclusive union observation. This represents an observation which is inclusively chosen from a set of named sub-observations.
@@ -462,11 +476,12 @@ public:
 	 * @param AttentionEncodingSize The encoding size used by the attention mechanism.
 	 * @param AttentionHeadNum The number of heads used by the attention mechanism.
 	 * @param ValueEncodingSize The output encoding size used by the attention mechanism.
+	 * @param bActive Whether to enable the observation as a valid network input. If set to inactive, this observation will be considered an empty observation. This flag allows easy toggles of specified observation.
 	 * @param Tag The tag of this new observation. Used during observation object validation and debugging.
 	 * @return The newly created observation schema element.
 	 */
 	UFUNCTION(BlueprintPure, Category = "LearningAgents", meta = (AdvancedDisplay = 2))
-	static UE_API FLearningAgentsObservationSchemaElement SpecifyInclusiveUnionObservation(ULearningAgentsObservationSchema* Schema, const TMap<FName, FLearningAgentsObservationSchemaElement>& Elements, const int32 AttentionEncodingSize = 32, const int32 AttentionHeadNum = 4, const int32 ValueEncodingSize = 32, const FName Tag = TEXT("InclusiveUnionObservation"));
+	static UE_API FLearningAgentsObservationSchemaElement SpecifyInclusiveUnionObservation(ULearningAgentsObservationSchema* Schema, const TMap<FName, FLearningAgentsObservationSchemaElement>& Elements, const int32 AttentionEncodingSize = 32, const int32 AttentionHeadNum = 4, const int32 ValueEncodingSize = 32, const bool bActive = true, const FName Tag = TEXT("InclusiveUnionObservation"));
 
 	/**
 	 * Specifies a new inclusive union observation. This represents an observation which is inclusively chosen from a set of named sub-observations.
@@ -480,11 +495,12 @@ public:
 	 * @param AttentionEncodingSize The encoding size used by the attention mechanism.
 	 * @param AttentionHeadNum The number of heads used by the attention mechanism.
 	 * @param ValueEncodingSize The output encoding size used by the attention mechanism.
+	 * @param bActive Whether to enable the observation as a valid network input. If set to inactive, this observation will be considered an empty observation. This flag allows easy toggles of specified observation.
 	 * @param Tag The tag of this new observation. Used during observation object validation and debugging.
 	 * @return The newly created observation schema element.
 	 */
 	UFUNCTION(BlueprintPure, Category = "LearningAgents", meta = (AdvancedDisplay = 3))
-	static UE_API FLearningAgentsObservationSchemaElement SpecifyInclusiveUnionObservationFromArrays(ULearningAgentsObservationSchema* Schema, const TArray<FName>& ElementNames, const TArray<FLearningAgentsObservationSchemaElement>& Elements, const int32 AttentionEncodingSize = 32, const int32 AttentionHeadNum = 4, const int32 ValueEncodingSize = 32, const FName Tag = TEXT("InclusiveUnionObservation"));
+	static UE_API FLearningAgentsObservationSchemaElement SpecifyInclusiveUnionObservationFromArrays(ULearningAgentsObservationSchema* Schema, const TArray<FName>& ElementNames, const TArray<FLearningAgentsObservationSchemaElement>& Elements, const int32 AttentionEncodingSize = 32, const int32 AttentionHeadNum = 4, const int32 ValueEncodingSize = 32, const bool bActive = true, const FName Tag = TEXT("InclusiveUnionObservation"));
 	
 	/**
 	 * Specifies a new inclusive union observation. This represents an observation which is inclusively chosen from a set of named sub-observations.
@@ -498,10 +514,11 @@ public:
 	 * @param AttentionEncodingSize The encoding size used by the attention mechanism.
 	 * @param AttentionHeadNum The number of heads used by the attention mechanism.
 	 * @param ValueEncodingSize The output encoding size used by the attention mechanism.
+	 * @param bActive Whether to enable the observation as a valid network input. If set to inactive, this observation will be considered an empty observation. This flag allows easy toggles of specified observation.
 	 * @param Tag The tag of this new observation. Used during observation object validation and debugging.
 	 * @return The newly created observation schema element.
 	 */
-	static UE_API FLearningAgentsObservationSchemaElement SpecifyInclusiveUnionObservationFromArrayViews(ULearningAgentsObservationSchema* Schema, const TArrayView<const FName> ElementNames, const TArrayView<const FLearningAgentsObservationSchemaElement> Elements, const int32 AttentionEncodingSize = 32, const int32 AttentionHeadNum = 4, const int32 ValueEncodingSize = 32, const FName Tag = TEXT("InclusiveUnionObservation"));
+	static UE_API FLearningAgentsObservationSchemaElement SpecifyInclusiveUnionObservationFromArrayViews(ULearningAgentsObservationSchema* Schema, const TArrayView<const FName> ElementNames, const TArrayView<const FLearningAgentsObservationSchemaElement> Elements, const int32 AttentionEncodingSize = 32, const int32 AttentionHeadNum = 4, const int32 ValueEncodingSize = 32, const bool bActive = true, const FName Tag = TEXT("InclusiveUnionObservation"));
 
 	/**
 	 * Specifies a new static array observation. This represents an observation made up of a fixed-size array of some other observation.
@@ -509,11 +526,12 @@ public:
 	 * @param Schema The Observation Schema
 	 * @param Element The sub-observation that represents elements of this array.
 	 * @param Num The number of elements in the fixed size array.
+	 * @param bActive Whether to enable the observation as a valid network input. If set to inactive, this observation will be considered an empty observation. This flag allows easy toggles of specified observation.
 	 * @param Tag The tag of this new observation. Used during observation object validation and debugging.
 	 * @return The newly created observation schema element.
 	 */
 	UFUNCTION(BlueprintPure, Category = "LearningAgents", meta = (AdvancedDisplay = 3))
-	static UE_API FLearningAgentsObservationSchemaElement SpecifyStaticArrayObservation(ULearningAgentsObservationSchema* Schema, const FLearningAgentsObservationSchemaElement Element, const int32 Num, const FName Tag = TEXT("StaticArrayObservation"));
+	static UE_API FLearningAgentsObservationSchemaElement SpecifyStaticArrayObservation(ULearningAgentsObservationSchema* Schema, const FLearningAgentsObservationSchemaElement Element, const int32 Num, const bool bActive = true, const FName Tag = TEXT("StaticArrayObservation"));
 
 	/**
 	 * Specifies a new set observation. This represents an observation made up of a Set of some other observation. This Set can be variable in size 
@@ -526,11 +544,12 @@ public:
 	 * @param AttentionEncodingSize The encoding size used by the attention mechanism.
 	 * @param AttentionHeadNum The number of heads used by the attention mechanism.
 	 * @param ValueEncodingSize The output encoding size used by the attention mechanism.
+	 * @param bActive Whether to enable the observation as a valid network input. If set to inactive, this observation will be considered an empty observation. This flag allows easy toggles of specified observation.
 	 * @param Tag The tag of this new observation. Used during observation object validation and debugging.
 	 * @return The newly created observation schema element.
 	 */
 	UFUNCTION(BlueprintPure, Category = "LearningAgents", meta = (AdvancedDisplay = 3))
-	static UE_API FLearningAgentsObservationSchemaElement SpecifySetObservation(ULearningAgentsObservationSchema* Schema, const FLearningAgentsObservationSchemaElement Element, const int32 MaxNum, const int32 AttentionEncodingSize = 32, const int32 AttentionHeadNum = 4, const int32 ValueEncodingSize = 32, const FName Tag = TEXT("SetObservation"));
+	static UE_API FLearningAgentsObservationSchemaElement SpecifySetObservation(ULearningAgentsObservationSchema* Schema, const FLearningAgentsObservationSchemaElement Element, const int32 MaxNum, const int32 AttentionEncodingSize = 32, const int32 AttentionHeadNum = 4, const int32 ValueEncodingSize = 32, const bool bActive = true, const FName Tag = TEXT("SetObservation"));
 
 	/**
 	 * Specifies a new pair observation. This represents an observation made up of two sub-observations.
@@ -538,11 +557,12 @@ public:
 	 * @param Schema The Observation Schema
 	 * @param Key The first sub-observation.
 	 * @param Value The second sub-observation.
+	 * @param bActive Whether to enable the observation as a valid network input. If set to inactive, this observation will be considered an empty observation. This flag allows easy toggles of specified observation.
 	 * @param Tag The tag of this new observation. Used during observation object validation and debugging.
 	 * @return The newly created observation schema element.
 	 */
 	UFUNCTION(BlueprintPure, Category = "LearningAgents", meta = (AdvancedDisplay = 3))
-	static UE_API FLearningAgentsObservationSchemaElement SpecifyPairObservation(ULearningAgentsObservationSchema* Schema, const FLearningAgentsObservationSchemaElement Key, const FLearningAgentsObservationSchemaElement Value, const FName Tag = TEXT("PairObservation"));
+	static UE_API FLearningAgentsObservationSchemaElement SpecifyPairObservation(ULearningAgentsObservationSchema* Schema, const FLearningAgentsObservationSchemaElement Key, const FLearningAgentsObservationSchemaElement Value, const bool bActive = true, const FName Tag = TEXT("PairObservation"));
 
 	/**
 	 * Specifies a new array observation. This represents an observation made up of an Array of some other observation. This Array can be variable in 
@@ -555,11 +575,12 @@ public:
 	 * @param AttentionEncodingSize The encoding size used by the attention mechanism.
 	 * @param AttentionHeadNum The number of heads used by the attention mechanism.
 	 * @param ValueEncodingSize The output encoding size used by the attention mechanism.
+	 * @param bActive Whether to enable the observation as a valid network input. If set to inactive, this observation will be considered an empty observation. This flag allows easy toggles of specified observation.
 	 * @param Tag The tag of this new observation. Used during observation object validation and debugging.
 	 * @return The newly created observation schema element.
 	 */
 	UFUNCTION(BlueprintPure, Category = "LearningAgents", meta = (AdvancedDisplay = 3))
-	static UE_API FLearningAgentsObservationSchemaElement SpecifyArrayObservation(ULearningAgentsObservationSchema* Schema, const FLearningAgentsObservationSchemaElement Element, const int32 MaxNum, const int32 AttentionEncodingSize = 32, const int32 AttentionHeadNum = 4, const int32 ValueEncodingSize = 32, const FName Tag = TEXT("ArrayObservation"));
+	static UE_API FLearningAgentsObservationSchemaElement SpecifyArrayObservation(ULearningAgentsObservationSchema* Schema, const FLearningAgentsObservationSchemaElement Element, const int32 MaxNum, const int32 AttentionEncodingSize = 32, const int32 AttentionHeadNum = 4, const int32 ValueEncodingSize = 32, const bool bActive = true, const FName Tag = TEXT("ArrayObservation"));
 
 	/**
 	 * Specifies a new map observation. This represents an observation made up of a Map of some other key and pair observations. This Map can be 
@@ -573,33 +594,38 @@ public:
 	 * @param AttentionEncodingSize The encoding size used by the attention mechanism.
 	 * @param AttentionHeadNum The number of heads used by the attention mechanism.
 	 * @param ValueEncodingSize The output encoding size used by the attention mechanism.
+	 * @param bActive Whether to enable the observation as a valid network input. If set to inactive, this observation will be considered an empty observation. This flag allows easy toggles of specified observation.
 	 * @param Tag The tag of this new observation. Used during observation object validation and debugging.
 	 * @return The newly created observation schema element.
 	 */
 	UFUNCTION(BlueprintPure, Category = "LearningAgents", meta = (AdvancedDisplay = 4))
-	static UE_API FLearningAgentsObservationSchemaElement SpecifyMapObservation(ULearningAgentsObservationSchema* Schema, const FLearningAgentsObservationSchemaElement KeyElement, const FLearningAgentsObservationSchemaElement ValueElement, const int32 MaxNum, const int32 AttentionEncodingSize = 32, const int32 AttentionHeadNum = 4, const int32 ValueEncodingSize = 32, const FName Tag = TEXT("MapObservation"));
+	static UE_API FLearningAgentsObservationSchemaElement SpecifyMapObservation(ULearningAgentsObservationSchema* Schema, const FLearningAgentsObservationSchemaElement KeyElement, const FLearningAgentsObservationSchemaElement ValueElement, const int32 MaxNum, const int32 AttentionEncodingSize = 32, const int32 AttentionHeadNum = 4, const int32 ValueEncodingSize = 32, const bool bActive = true, const FName Tag = TEXT("MapObservation"));
 
 	/**
-	 * Specifies a new enum observation. This represents an exclusive choice from elements of the given Enum. To use this with an Enum defined in C++ 
+	 * Specifies a new enum observation. This represents an exclusive choice from elements of the given Enum. To use this with an Enum defined in C++
 	 * use the FindEnumByName convenience function.
 	 *
 	 * @param Schema The Observation Schema
 	 * @param Enum The enum type to use.
+	 * @param bActive Whether to enable the observation as a valid network input. If set to inactive, this observation will be considered an empty observation. This flag allows easy toggles of specified observation.
+	 * @param Tag The tag of this new observation. Used during observation object validation and debugging.
 	 * @return The newly created observation schema element.
 	 */
 	UFUNCTION(BlueprintPure, Category = "LearningAgents", meta = (AdvancedDisplay = 2))
-	static UE_API FLearningAgentsObservationSchemaElement SpecifyEnumObservation(ULearningAgentsObservationSchema* Schema, const UEnum* Enum, const FName Tag = TEXT("EnumObservation"));
+	static UE_API FLearningAgentsObservationSchemaElement SpecifyEnumObservation(ULearningAgentsObservationSchema* Schema, const UEnum* Enum, const bool bActive = true, const FName Tag = TEXT("EnumObservation"));
 
 	/**
-	 * Specifies a new bitmask observation. This represents an inclusive choice from elements of the given Enum. To use this with an Enum defined in 
+	 * Specifies a new bitmask observation. This represents an inclusive choice from elements of the given Enum. To use this with an Enum defined in
 	 * C++ use the FindEnumByName convenience function.
 	 *
 	 * @param Schema The Observation Schema
 	 * @param Enum The enum type to use.
+	 * @param bActive Whether to enable the observation as a valid network input. If set to inactive, this observation will be considered an empty observation. This flag allows easy toggles of specified observation.
+	 * @param Tag The tag of this new observation. Used during observation object validation and debugging.
 	 * @return The newly created observation schema element.
 	 */
 	UFUNCTION(BlueprintPure, Category = "LearningAgents", meta = (AdvancedDisplay = 2))
-	static UE_API FLearningAgentsObservationSchemaElement SpecifyBitmaskObservation(ULearningAgentsObservationSchema* Schema, const UEnum* Enum, const FName Tag = TEXT("BitmaskObservation"));
+	static UE_API FLearningAgentsObservationSchemaElement SpecifyBitmaskObservation(ULearningAgentsObservationSchema* Schema, const UEnum* Enum, const bool bActive = true, const FName Tag = TEXT("BitmaskObservation"));
 
 	/**
 	 * Specifies a new optional observation. This represents an observation which may or may not be provided.
@@ -607,10 +633,12 @@ public:
 	 * @param Schema The Observation Schema
 	 * @param Elements The sub-observation that may or may not be provided.
 	 * @param EncodingSize The encoding size used to encode this sub-observation.
+	 * @param bActive Whether to enable the observation as a valid network input. If set to inactive, this observation will be considered an empty observation. This flag allows easy toggles of specified observation.
+	 * @param Tag The tag of this new observation. Used during observation object validation and debugging.
 	 * @return The newly created observation schema element.
 	 */
 	UFUNCTION(BlueprintPure, Category = "LearningAgents", meta = (AdvancedDisplay = 2))
-	static UE_API FLearningAgentsObservationSchemaElement SpecifyOptionalObservation(ULearningAgentsObservationSchema* Schema, const FLearningAgentsObservationSchemaElement Element, const int32 EncodingSize = 128, const FName Tag = TEXT("OptionalObservation"));
+	static UE_API FLearningAgentsObservationSchemaElement SpecifyOptionalObservation(ULearningAgentsObservationSchema* Schema, const FLearningAgentsObservationSchemaElement Element, const int32 EncodingSize = 128, const bool bActive = true, const FName Tag = TEXT("OptionalObservation"));
 
 	/**
 	 * Specifies a new either observation. This represents an observation which will be either sub-observation A or sub-observation B.
@@ -619,10 +647,12 @@ public:
 	 * @param A The first sub-observation.
 	 * @param A The second sub-observation.
 	 * @param EncodingSize The encoding size used to encode each sub-observation.
+	 * @param bActive Whether to enable the observation as a valid network input. If set to inactive, this observation will be considered an empty observation. This flag allows easy toggles of specified observation.
+	 * @param Tag The tag of this new observation. Used during observation object validation and debugging.
 	 * @return The newly created observation schema element.
 	 */
 	UFUNCTION(BlueprintPure, Category = "LearningAgents", meta = (AdvancedDisplay = 3))
-	static UE_API FLearningAgentsObservationSchemaElement SpecifyEitherObservation(ULearningAgentsObservationSchema* Schema, const FLearningAgentsObservationSchemaElement A, const FLearningAgentsObservationSchemaElement B, const int32 EncodingSize = 128, const FName Tag = TEXT("EitherObservation"));
+	static UE_API FLearningAgentsObservationSchemaElement SpecifyEitherObservation(ULearningAgentsObservationSchema* Schema, const FLearningAgentsObservationSchemaElement A, const FLearningAgentsObservationSchemaElement B, const int32 EncodingSize = 128, const bool bActive = true, const FName Tag = TEXT("EitherObservation"));
 
 	/**
 	 * Specifies a new encoding observation. This represents an observation which will be an encoding of another sub-observation using a small neural
@@ -633,20 +663,23 @@ public:
 	 * @param EncodingSize The encoding size used to encode this sub-observation.
 	 * @param HiddenLayerNum The number of hidden layers used to encode this sub-observation.
 	 * @param ActivationFunction The activation function used to encode this sub-observation.
+	 * @param bActive Whether to enable the observation as a valid network input. If set to inactive, this observation will be considered an empty observation. This flag allows easy toggles of specified observation.
+	 * @param Tag The tag of this new observation. Used during observation object validation and debugging.
 	 * @return The newly created observation schema element.
 	 */
 	UFUNCTION(BlueprintPure, Category = "LearningAgents", meta = (AdvancedDisplay = 2))
-	static UE_API FLearningAgentsObservationSchemaElement SpecifyEncodingObservation(ULearningAgentsObservationSchema* Schema, const FLearningAgentsObservationSchemaElement Element, const int32 EncodingSize = 128, const int32 HiddenLayerNum = 1, const ELearningAgentsActivationFunction ActivationFunction = ELearningAgentsActivationFunction::ELU, const FName Tag = TEXT("EncodingObservation"));
+	static UE_API FLearningAgentsObservationSchemaElement SpecifyEncodingObservation(ULearningAgentsObservationSchema* Schema, const FLearningAgentsObservationSchemaElement Element, const int32 EncodingSize = 128, const int32 HiddenLayerNum = 1, const ELearningAgentsActivationFunction ActivationFunction = ELearningAgentsActivationFunction::ELU, const bool bActive = true, const FName Tag = TEXT("EncodingObservation"));
 
 	/**
 	 * Specifies a new bool observation. A true or false observation.
 	 *
 	 * @param Schema The Observation Schema
+	 * @param bActive Whether to enable the observation as a valid network input. If set to inactive, this observation will be considered an empty observation. This flag allows easy toggles of specified observation.
 	 * @param Tag The tag of this new observation. Used during observation object validation and debugging.
 	 * @return The newly created observation schema element.
 	 */
 	UFUNCTION(BlueprintPure, Category = "LearningAgents", meta = (AdvancedDisplay = 1))
-	static UE_API FLearningAgentsObservationSchemaElement SpecifyBoolObservation(ULearningAgentsObservationSchema* Schema, const FName Tag = TEXT("BoolObservation"));
+	static UE_API FLearningAgentsObservationSchemaElement SpecifyBoolObservation(ULearningAgentsObservationSchema* Schema, const bool bActive = true, const FName Tag = TEXT("BoolObservation"));
 
 	/**
 	 * Specifies a new 1d convolution observation of another observation (Flat buffer input, Channel x Length).
@@ -654,11 +687,12 @@ public:
 	 * @param Schema The Observation Schema
 	 * @param Element The sub-observation to be convolved.
 	 * @param FConv1dObservationParams The convolution parameters.
+	 * @param bActive Whether to enable the observation as a valid network input. If set to inactive, this observation will be considered an empty observation. This flag allows easy toggles of specified observation.
 	 * @param Tag The tag of this new observation. Used during observation object validation and debugging.
 	 * @return The newly created observation schema element.
 	 */
 	UFUNCTION(BlueprintPure, Category = "LearningAgents", meta = (AdvancedDisplay = 9))
-	static UE_API FLearningAgentsObservationSchemaElement SpecifyConv1dObservation(ULearningAgentsObservationSchema* Schema, const FLearningAgentsObservationSchemaElement Element, const FConv1dObservationParams& Conv1dParams, const FName Tag = TEXT("Conv1dObservation"));
+	static UE_API FLearningAgentsObservationSchemaElement SpecifyConv1dObservation(ULearningAgentsObservationSchema* Schema, const FLearningAgentsObservationSchemaElement Element, const FConv1dObservationParams& Conv1dParams, const bool bActive = true, const FName Tag = TEXT("Conv1dObservation"));
 
 	/**
 	 * Specifies a new 2d convolution observation of another observation (Flat buffer input, Channels x Input Height x Input Width), Zeros Padding.
@@ -666,11 +700,12 @@ public:
 	 * @param Schema The Observation Schema
 	 * @param Element The sub-observation to be convolved.
 	 * @param FConv2dObservationParams The convolution parameters.
+	 * @param bActive Whether to enable the observation as a valid network input. If set to inactive, this observation will be considered an empty observation. This flag allows easy toggles of specified observation.
 	 * @param Tag The tag of this new observation. Used during observation object validation and debugging.
 	 * @return The newly created observation schema element.
 	 */
 	 UFUNCTION(BlueprintPure, Category = "LearningAgents", meta = (AdvancedDisplay = 10))
-	static UE_API FLearningAgentsObservationSchemaElement SpecifyConv2dObservation(ULearningAgentsObservationSchema* Schema, const FLearningAgentsObservationSchemaElement Element, const FConv2dObservationParams& Conv2dParams, const FName Tag = TEXT("Conv2dObservation"));
+	static UE_API FLearningAgentsObservationSchemaElement SpecifyConv2dObservation(ULearningAgentsObservationSchema* Schema, const FLearningAgentsObservationSchemaElement Element, const FConv2dObservationParams& Conv2dParams, const bool bActive = true, const FName Tag = TEXT("Conv2dObservation"));
 
 	/**
 	 * Specifies a new float observation. A simple observation which can be used as a catch-all for situations where a
@@ -678,142 +713,155 @@ public:
 	 *
 	 * @param Schema The Observation Schema
 	 * @param FloatScale The scale used to normalize this observation.
+	 * @param bActive Whether to enable the observation as a valid network input. If set to inactive, this observation will be considered an empty observation. This flag allows easy toggles of specified observation.
 	 * @param Tag The tag of this new observation. Used during observation object validation and debugging.
 	 * @return The newly created observation schema element.
 	 */
 	UFUNCTION(BlueprintPure, Category = "LearningAgents", meta = (AdvancedDisplay = 2))
-	static UE_API FLearningAgentsObservationSchemaElement SpecifyFloatObservation(ULearningAgentsObservationSchema* Schema, const float FloatScale = 1.0f, const FName Tag = TEXT("FloatObservation"));
+	static UE_API FLearningAgentsObservationSchemaElement SpecifyFloatObservation(ULearningAgentsObservationSchema* Schema, const float FloatScale = 1.0f, const bool bActive = true, const FName Tag = TEXT("FloatObservation"));
 
 	/**
 	 * Specifies a new location observation. Allows an agent to observe the location of some entity.
 	 *
 	 * @param Schema The Observation Schema
 	 * @param LocationScale The scale used to normalize this observation in cm.
+	 * @param bActive Whether to enable the observation as a valid network input. If set to inactive, this observation will be considered an empty observation. This flag allows easy toggles of specified observation.
 	 * @param Tag The tag of this new observation. Used during observation object validation and debugging.
 	 * @return The newly created observation schema element.
 	 */
 	UFUNCTION(BlueprintPure, Category = "LearningAgents", meta = (AdvancedDisplay = 2))
-	static UE_API FLearningAgentsObservationSchemaElement SpecifyLocationObservation(ULearningAgentsObservationSchema* Schema, const float LocationScale = 100.0f, const FName Tag = TEXT("LocationObservation"));
+	static UE_API FLearningAgentsObservationSchemaElement SpecifyLocationObservation(ULearningAgentsObservationSchema* Schema, const float LocationScale = 100.0f, const bool bActive = true, const FName Tag = TEXT("LocationObservation"));
 
 	/**
-	 * Specifies a new rotation observation. Allows an agent to observe the rotation of some entity. Rotations are encoded as two columns of the 
+	 * Specifies a new rotation observation. Allows an agent to observe the rotation of some entity. Rotations are encoded as two columns of the
 	 * rotation matrix to ensure there is no discontinuity in the encoding.
 	 *
 	 * @param Schema The Observation Schema
+	 * @param bActive Whether to enable the observation as a valid network input. If set to inactive, this observation will be considered an empty observation. This flag allows easy toggles of specified observation.
 	 * @param Tag The tag of this new observation. Used during observation object validation and debugging.
 	 * @return The newly created observation schema element.
 	 */
 	UFUNCTION(BlueprintPure, Category = "LearningAgents", meta = (AdvancedDisplay = 1))
-	static UE_API FLearningAgentsObservationSchemaElement SpecifyRotationObservation(ULearningAgentsObservationSchema* Schema, const FName Tag = TEXT("RotationObservation"));
+	static UE_API FLearningAgentsObservationSchemaElement SpecifyRotationObservation(ULearningAgentsObservationSchema* Schema, const bool bActive = true, const FName Tag = TEXT("RotationObservation"));
 
 	/**
-	 * Specifies a new scale observation. Allows an agent to observe the scale of some entity. Negative scales are not supported by this observation 
+	 * Specifies a new scale observation. Allows an agent to observe the scale of some entity. Negative scales are not supported by this observation
 	 * type.
 	 *
 	 * @param Schema The Observation Schema
+	 * @param bActive Whether to enable the observation as a valid network input. If set to inactive, this observation will be considered an empty observation. This flag allows easy toggles of specified observation.
 	 * @param Tag The tag of this new observation. Used during observation object validation and debugging.
 	 * @return The newly created observation schema element.
 	 */
 	UFUNCTION(BlueprintPure, Category = "LearningAgents", meta = (AdvancedDisplay = 1))
-	static UE_API FLearningAgentsObservationSchemaElement SpecifyScaleObservation(ULearningAgentsObservationSchema* Schema, const FName Tag = TEXT("ScaleObservation"));
+	static UE_API FLearningAgentsObservationSchemaElement SpecifyScaleObservation(ULearningAgentsObservationSchema* Schema, const bool bActive = true, const FName Tag = TEXT("ScaleObservation"));
 
 	/**
 	 * Specifies a new transform observation. Allows an agent to observe the transform of some entity.
 	 *
 	 * @param Schema The Observation Schema
 	 * @param LocationScale The scale used to normalize the Location part of this transform observation in cm.
+	 * @param bActive Whether to enable the observation as a valid network input. If set to inactive, this observation will be considered an empty observation. This flag allows easy toggles of specified observation.
 	 * @param Tag The tag of this new observation. Used during observation object validation and debugging.
 	 * @return The newly created observation schema element.
 	 */
 	UFUNCTION(BlueprintPure, Category = "LearningAgents", meta = (AdvancedDisplay = 2))
-	static UE_API FLearningAgentsObservationSchemaElement SpecifyTransformObservation(ULearningAgentsObservationSchema* Schema, const float LocationScale = 100.0f, const FName Tag = TEXT("TransformObservation"));
+	static UE_API FLearningAgentsObservationSchemaElement SpecifyTransformObservation(ULearningAgentsObservationSchema* Schema, const float LocationScale = 100.0f, const bool bActive = true, const FName Tag = TEXT("TransformObservation"));
 
 	/**
-	 * Specifies a new angle observation. This will be encoded as a 2-dimension Cartesian vector so that 0 and 350 are close to each other in the 
+	 * Specifies a new angle observation. This will be encoded as a 2-dimension Cartesian vector so that 0 and 350 are close to each other in the
 	 * encoded space.
 	 *
 	 * @param Schema The Observation Schema
+	 * @param bActive Whether to enable the observation as a valid network input. If set to inactive, this observation will be considered an empty observation. This flag allows easy toggles of specified observation.
 	 * @param Tag The tag of this new observation. Used during observation object validation and debugging.
 	 * @return The newly created observation schema element.
 	 */
 	UFUNCTION(BlueprintPure, Category = "LearningAgents", meta = (AdvancedDisplay = 1))
-	static UE_API FLearningAgentsObservationSchemaElement SpecifyAngleObservation(ULearningAgentsObservationSchema* Schema, const FName Tag = TEXT("AngleObservation"));
+	static UE_API FLearningAgentsObservationSchemaElement SpecifyAngleObservation(ULearningAgentsObservationSchema* Schema, const bool bActive = true, const FName Tag = TEXT("AngleObservation"));
 
 	/**
 	 * Specifies a new velocity observation. Allows an agent to observe the velocity of some entity.
 	 *
 	 * @param Schema The Observation Schema
 	 * @param VelocityScale The scale used to normalize this observation in cm/s.
+	 * @param bActive Whether to enable the observation as a valid network input. If set to inactive, this observation will be considered an empty observation. This flag allows easy toggles of specified observation.
 	 * @param Tag The tag of this new observation. Used during observation object validation and debugging.
 	 * @return The newly created observation schema element.
 	 */
 	UFUNCTION(BlueprintPure, Category = "LearningAgents", meta = (AdvancedDisplay = 2))
-	static UE_API FLearningAgentsObservationSchemaElement SpecifyVelocityObservation(ULearningAgentsObservationSchema* Schema, const float VelocityScale = 200.0f, const FName Tag = TEXT("VelocityObservation"));
+	static UE_API FLearningAgentsObservationSchemaElement SpecifyVelocityObservation(ULearningAgentsObservationSchema* Schema, const float VelocityScale = 200.0f, const bool bActive = true, const FName Tag = TEXT("VelocityObservation"));
 
 	/**
 	 * Specifies a new direction observation. Allows an agent to observe the direction of some entity.
 	 *
 	 * @param Schema The Observation Schema
+	 * @param bActive Whether to enable the observation as a valid network input. If set to inactive, this observation will be considered an empty observation. This flag allows easy toggles of specified observation.
 	 * @param Tag The tag of this new observation. Used during observation object validation and debugging.
 	 * @return The newly created observation schema element.
 	 */
 	UFUNCTION(BlueprintPure, Category = "LearningAgents", meta = (AdvancedDisplay = 1))
-	static UE_API FLearningAgentsObservationSchemaElement SpecifyDirectionObservation(ULearningAgentsObservationSchema* Schema, const FName Tag = TEXT("DirectionObservation"));
+	static UE_API FLearningAgentsObservationSchemaElement SpecifyDirectionObservation(ULearningAgentsObservationSchema* Schema, const bool bActive = true, const FName Tag = TEXT("DirectionObservation"));
 
 	/**
 	 * Specifies a new location along spline observation. This observes the location of the spline at the given distance along that spline.
 	 *
 	 * @param Schema The Observation Schema
 	 * @param LocationScale The scale used to normalize this observation in cm.
+	 * @param bActive Whether to enable the observation as a valid network input. If set to inactive, this observation will be considered an empty observation. This flag allows easy toggles of specified observation.
 	 * @param Tag The tag of this new observation. Used during observation object validation and debugging.
 	 * @return The newly created observation schema element.
 	 */
 	UFUNCTION(BlueprintPure, Category = "LearningAgents", meta = (AdvancedDisplay = 2))
-	static UE_API FLearningAgentsObservationSchemaElement SpecifyLocationAlongSplineObservation(ULearningAgentsObservationSchema* Schema, const float LocationScale = 100.0f, const FName Tag = TEXT("LocationAlongSplineObservation"));
+	static UE_API FLearningAgentsObservationSchemaElement SpecifyLocationAlongSplineObservation(ULearningAgentsObservationSchema* Schema, const float LocationScale = 100.0f, const bool bActive = true, const FName Tag = TEXT("LocationAlongSplineObservation"));
 
 	/**
-	 * Specifies a new proportion along spline observation. This observes the proportion along a spline at the given distance. For looped splines 
-	 * this will be treated effectively like an angle between 0 and 360 degrees and encoded appropriately so that 0 and 350 are close to each other in the 
+	 * Specifies a new proportion along spline observation. This observes the proportion along a spline at the given distance. For looped splines
+	 * this will be treated effectively like an angle between 0 and 360 degrees and encoded appropriately so that 0 and 350 are close to each other in the
 	 * encoded space, while for non-looped splines this will be treated as a value between 0 and 1.
 	 *
 	 * @param Schema The Observation Schema
+	 * @param bActive Whether to enable the observation as a valid network input. If set to inactive, this observation will be considered an empty observation. This flag allows easy toggles of specified observation.
 	 * @param Tag The tag of this new observation. Used during observation object validation and debugging.
 	 * @return The newly created observation schema element.
 	 */
 	UFUNCTION(BlueprintPure, Category = "LearningAgents", meta = (AdvancedDisplay = 1))
-	static UE_API FLearningAgentsObservationSchemaElement SpecifyProportionAlongSplineObservation(ULearningAgentsObservationSchema* Schema, const FName Tag = TEXT("ProportionAlongSplineObservation"));
+	static UE_API FLearningAgentsObservationSchemaElement SpecifyProportionAlongSplineObservation(ULearningAgentsObservationSchema* Schema, const bool bActive = true, const FName Tag = TEXT("ProportionAlongSplineObservation"));
 
 	/**
 	 * Specifies a new direction along spline observation. This observes the direction of the spline at the given distance along that spline.
 	 *
 	 * @param Schema The Observation Schema
+	 * @param bActive Whether to enable the observation as a valid network input. If set to inactive, this observation will be considered an empty observation. This flag allows easy toggles of specified observation.
 	 * @param Tag The tag of this new observation. Used during observation object validation and debugging.
 	 * @return The newly created observation schema element.
 	 */
 	UFUNCTION(BlueprintPure, Category = "LearningAgents", meta = (AdvancedDisplay = 1))
-	static UE_API FLearningAgentsObservationSchemaElement SpecifyDirectionAlongSplineObservation(ULearningAgentsObservationSchema* Schema, const FName Tag = TEXT("DirectionAlongSplineObservation"));
+	static UE_API FLearningAgentsObservationSchemaElement SpecifyDirectionAlongSplineObservation(ULearningAgentsObservationSchema* Schema, const bool bActive = true, const FName Tag = TEXT("DirectionAlongSplineObservation"));
 
 	/**
-	 * Specifies a new proportion along ray observation. This observes how far a you can travel along a ray before collision. Rays that can travel 
+	 * Specifies a new proportion along ray observation. This observes how far a you can travel along a ray before collision. Rays that can travel
 	 * the full distance are encoded as zero, while rays that collide instantly are encoded as one.
 	 *
 	 * @param Schema The Observation Schema
+	 * @param bActive Whether to enable the observation as a valid network input. If set to inactive, this observation will be considered an empty observation. This flag allows easy toggles of specified observation.
 	 * @param Tag The tag of this new observation. Used during observation object validation and debugging.
 	 * @return The newly created observation schema element.
 	 */
 	UFUNCTION(BlueprintPure, Category = "LearningAgents", meta = (AdvancedDisplay = 1))
-	static UE_API FLearningAgentsObservationSchemaElement SpecifyProportionAlongRayObservation(ULearningAgentsObservationSchema* Schema, const FName Tag = TEXT("ProportionAlongRayObservation"));
+	static UE_API FLearningAgentsObservationSchemaElement SpecifyProportionAlongRayObservation(ULearningAgentsObservationSchema* Schema, const bool bActive = true, const FName Tag = TEXT("ProportionAlongRayObservation"));
 	
 	/**
-	 * Specifies a new LiDAR observation. This observes how far a you can travel along a ray before collision. Rays that can travel 
+	 * Specifies a new LiDAR observation. This observes how far a you can travel along a ray before collision. Rays that can travel
 	 * the full distance are encoded as one, while rays that collide instantly are encoded as zero.
 	 *
 	 * @param Schema The Observation Schema
+	 * @param bActive Whether to enable the observation as a valid network input. If set to inactive, this observation will be considered an empty observation. This flag allows easy toggles of specified observation.
 	 * @param Tag The tag of this new observation. Used during observation object validation and debugging.
 	 * @return The newly created observation schema element.
 	 */
 	UFUNCTION(BlueprintPure, Category = "LearningAgents", meta = (AdvancedDisplay = 1))
-	static UE_API FLearningAgentsObservationSchemaElement SpecifyLidarObservation(ULearningAgentsObservationSchema* Schema, const FName Tag = TEXT("LidarObservation"));
+	static UE_API FLearningAgentsObservationSchemaElement SpecifyLidarObservation(ULearningAgentsObservationSchema* Schema, const bool bActive = true, const FName Tag = TEXT("LidarObservation"));
 
 public:
 
@@ -832,6 +880,7 @@ public:
 	 *
 	 * @param Object The Observation Object
 	 * @param Values The observation values.
+	 * @param bActive Whether the observation should be considered active. Must match the value given during Specify.
 	 * @param Tag The tag of the corresponding observation. Must match the tag given during Specify.
 	 * @param bVisualLoggerEnabled When true, debug data will be sent to the visual logger.
 	 * @param VisualLoggerListener The listener object which is making this observation. This must be set to use logging.
@@ -842,8 +891,9 @@ public:
 	 */
 	UFUNCTION(BlueprintPure, Category = "LearningAgents", meta = (AdvancedDisplay = 2, DefaultToSelf = "VisualLoggerListener"))
 	static UE_API FLearningAgentsObservationObjectElement MakeContinuousObservation(
-		ULearningAgentsObservationObject* Object, 
-		const TArray<float>& Values, 
+		ULearningAgentsObservationObject* Object,
+		const TArray<float>& Values,
+		const bool bActive = true,
 		const FName Tag = TEXT("ContinuousObservation"),
 		const bool bVisualLoggerEnabled = false,
 		ULearningAgentsManagerListener* VisualLoggerListener = nullptr,
@@ -856,6 +906,7 @@ public:
 	 *
 	 * @param Object The Observation Object
 	 * @param Values The observation values.
+	 * @param bActive Whether the observation should be considered active. Must match the value given during Specify.
 	 * @param Tag The tag of the corresponding observation. Must match the tag given during Specify.
 	 * @param bVisualLoggerEnabled When true, debug data will be sent to the visual logger.
 	 * @param VisualLoggerListener The listener object which is making this observation. This must be set to use logging.
@@ -867,6 +918,7 @@ public:
 	static UE_API FLearningAgentsObservationObjectElement MakeContinuousObservationFromArrayView(
 		ULearningAgentsObservationObject* Object, 
 		const TArrayView<const float> Values, 
+		const bool bActive = true, 
 		const FName Tag = TEXT("ContinuousObservation"),
 		const bool bVisualLoggerEnabled = false,
 		ULearningAgentsManagerListener* VisualLoggerListener = nullptr,
@@ -879,6 +931,7 @@ public:
 	 *
 	 * @param Object The Observation Object
 	 * @param DiscreteIndex The index of the discrete observation.
+	 * @param bActive Whether the observation should be considered active. Must match the value given during Specify.
 	 * @param Tag The tag of the corresponding observation. Must match the tag given during Specify.
 	 * @param bVisualLoggerEnabled When true, debug data will be sent to the visual logger.
 	 * @param VisualLoggerListener The listener object which is making this observation. This must be set to use logging.
@@ -891,6 +944,7 @@ public:
 	static UE_API FLearningAgentsObservationObjectElement MakeExclusiveDiscreteObservation(
 		ULearningAgentsObservationObject* Object, 
 		const int32 DiscreteIndex, 
+		const bool bActive = true, 
 		const FName Tag = TEXT("ExclusiveDiscreteObservation"),
 		const bool bVisualLoggerEnabled = false,
 		ULearningAgentsManagerListener* VisualLoggerListener = nullptr,
@@ -903,6 +957,7 @@ public:
 	 *
 	 * @param Object The Observation Object
 	 * @param DiscreteIndices The indices of the discrete observations.
+	 * @param bActive Whether the observation should be considered active. Must match the value given during Specify.
 	 * @param Tag The tag of the corresponding observation. Must match the tag given during Specify.
 	 * @param bVisualLoggerEnabled When true, debug data will be sent to the visual logger.
 	 * @param VisualLoggerListener The listener object which is making this observation. This must be set to use logging.
@@ -915,6 +970,7 @@ public:
 	static UE_API FLearningAgentsObservationObjectElement MakeInclusiveDiscreteObservation(
 		ULearningAgentsObservationObject* Object, 
 		const TArray<int32>& DiscreteIndices, 
+		const bool bActive = true, 
 		const FName Tag = TEXT("InclusiveDiscreteObservation"),
 		const bool bVisualLoggerEnabled = false,
 		ULearningAgentsManagerListener* VisualLoggerListener = nullptr,
@@ -927,6 +983,7 @@ public:
 	 *
 	 * @param Object The Observation Object
 	 * @param DiscreteIndices The indices of the discrete observations.
+	 * @param bActive Whether the observation should be considered active. Must match the value given during Specify.
 	 * @param Tag The tag of the corresponding observation. Must match the tag given during Specify.
 	 * @param bVisualLoggerEnabled When true, debug data will be sent to the visual logger.
 	 * @param VisualLoggerListener The listener object which is making this observation. This must be set to use logging.
@@ -938,6 +995,7 @@ public:
 	static UE_API FLearningAgentsObservationObjectElement MakeInclusiveDiscreteObservationFromArrayView(
 		ULearningAgentsObservationObject* Object, 
 		const TArrayView<const int32> DiscreteIndices, 
+		const bool bActive = true, 
 		const FName Tag = TEXT("InclusiveDiscreteObservation"),
 		const bool bVisualLoggerEnabled = false,
 		ULearningAgentsManagerListener* VisualLoggerListener = nullptr,
@@ -950,6 +1008,7 @@ public:
 	 *
 	 * @param Object The Observation Object
 	 * @param ElementName The name of the discrete observation.
+	 * @param bActive Whether the observation should be considered active. Must match the value given during Specify.
 	 * @param Tag The tag of the corresponding observation. Must match the tag given during Specify.
 	 * @param bVisualLoggerEnabled When true, debug data will be sent to the visual logger.
 	 * @param VisualLoggerListener The listener object which is making this observation. This must be set to use logging.
@@ -962,6 +1021,7 @@ public:
 	static UE_API FLearningAgentsObservationObjectElement MakeNamedExclusiveDiscreteObservation(
 		ULearningAgentsObservationObject* Object,
 		const FName ElementName,
+		const bool bActive = true, 
 		const FName Tag = TEXT("NamedExclusiveDiscreteObservation"),
 		const bool bVisualLoggerEnabled = false,
 		ULearningAgentsManagerListener* VisualLoggerListener = nullptr,
@@ -974,6 +1034,7 @@ public:
 	 *
 	 * @param Object The Observation Object
 	 * @param ElementNames The names of the discrete observations.
+	 * @param bActive Whether the observation should be considered active. Must match the value given during Specify.
 	 * @param Tag The tag of the corresponding observation. Must match the tag given during Specify.
 	 * @param bVisualLoggerEnabled When true, debug data will be sent to the visual logger.
 	 * @param VisualLoggerListener The listener object which is making this observation. This must be set to use logging.
@@ -986,6 +1047,7 @@ public:
 	static UE_API FLearningAgentsObservationObjectElement MakeNamedInclusiveDiscreteObservation(
 		ULearningAgentsObservationObject* Object,
 		const TArray<FName>& ElementNames,
+		const bool bActive = true, 
 		const FName Tag = TEXT("NamedInclusiveDiscreteObservation"),
 		const bool bVisualLoggerEnabled = false,
 		ULearningAgentsManagerListener* VisualLoggerListener = nullptr,
@@ -997,7 +1059,8 @@ public:
 	 * Make a new named inclusive discrete observation.
 	 *
 	 * @param Object The Observation Object
-	 * @param ElementNames The names of the discrete observations. 
+	 * @param ElementNames The names of the discrete observations.
+	 * @param bActive Whether the observation should be considered active. Must match the value given during Specify.
 	 * @param Tag The tag of the corresponding observation. Must match the tag given during Specify.
 	 * @param bVisualLoggerEnabled When true, debug data will be sent to the visual logger.
 	 * @param VisualLoggerListener The listener object which is making this observation. This must be set to use logging.
@@ -1009,6 +1072,7 @@ public:
 	static UE_API FLearningAgentsObservationObjectElement MakeNamedInclusiveDiscreteObservationFromArrayView(
 		ULearningAgentsObservationObject* Object,
 		const TArrayView<const FName> ElementNames,
+		const bool bActive = true, 
 		const FName Tag = TEXT("NamedInclusiveDiscreteObservation"),
 		const bool bVisualLoggerEnabled = false,
 		ULearningAgentsManagerListener* VisualLoggerListener = nullptr,
@@ -1022,6 +1086,7 @@ public:
 	 * @param Object The Observation Object
 	 * @param Num The number of items. Must be less than or equal to MaxNum.
 	 * @param MaxNum The maximum number of items possible.
+	 * @param bActive Whether the observation should be considered active. Must match the value given during Specify.
 	 * @param Tag The tag of the corresponding observation. Must match the tag given during Specify.
 	 * @param bVisualLoggerEnabled When true, debug data will be sent to the visual logger.
 	 * @param VisualLoggerListener The listener object which is making this observation. This must be set to use logging.
@@ -1035,6 +1100,7 @@ public:
 		ULearningAgentsObservationObject* Object, 
 		const int32 Num, 
 		const int32 MaxNum, 
+		const bool bActive = true, 
 		const FName Tag = TEXT("CountObservation"),
 		const bool bVisualLoggerEnabled = false,
 		ULearningAgentsManagerListener* VisualLoggerListener = nullptr,
@@ -1047,11 +1113,12 @@ public:
 	 *
 	 * @param Object The Observation Object
 	 * @param Elements The named sub-observations. Must match what was given during Specify.
+	 * @param bActive Whether the observation should be considered active. Must match the value given during Specify.
 	 * @param Tag The tag of the corresponding observation. Must match the tag given during Specify.
 	 * @return The newly created observation object element.
 	 */
 	UFUNCTION(BlueprintPure, Category = "LearningAgents", meta = (AdvancedDisplay = 2))
-	static UE_API FLearningAgentsObservationObjectElement MakeStructObservation(ULearningAgentsObservationObject* Object, const TMap<FName, FLearningAgentsObservationObjectElement>& Elements, const FName Tag = TEXT("StructObservation"));
+	static UE_API FLearningAgentsObservationObjectElement MakeStructObservation(ULearningAgentsObservationObject* Object, const TMap<FName, FLearningAgentsObservationObjectElement>& Elements, const bool bActive = true, const FName Tag = TEXT("StructObservation"));
 
 	/**
 	 * Make a new struct observation.
@@ -1059,11 +1126,12 @@ public:
 	 * @param Object The Observation Object
 	 * @param ElementNames The names of the sub-observations. Must match what was given during Specify.
 	 * @param Elements The corresponding sub-observations. Must match what was given during Specify.
+	 * @param bActive Whether the observation should be considered active. Must match the value given during Specify.
 	 * @param Tag The tag of the corresponding observation. Must match the tag given during Specify.
 	 * @return The newly created observation object element.
 	 */
 	UFUNCTION(BlueprintPure, Category = "LearningAgents", meta = (AdvancedDisplay = 3))
-	static UE_API FLearningAgentsObservationObjectElement MakeStructObservationFromArrays(ULearningAgentsObservationObject* Object, const TArray<FName>& ElementNames, const TArray<FLearningAgentsObservationObjectElement>& Elements, const FName Tag = TEXT("StructObservation"));
+	static UE_API FLearningAgentsObservationObjectElement MakeStructObservationFromArrays(ULearningAgentsObservationObject* Object, const TArray<FName>& ElementNames, const TArray<FLearningAgentsObservationObjectElement>& Elements, const bool bActive = true, const FName Tag = TEXT("StructObservation"));
 	
 	/**
 	 * Make a new struct observation.
@@ -1071,10 +1139,11 @@ public:
 	 * @param Object The Observation Object
 	 * @param ElementNames The names of the sub-observations. Must match what was given during Specify.
 	 * @param Elements The corresponding sub-observations. Must match what was given during Specify.
+	 * @param bActive Whether the observation should be considered active. Must match the value given during Specify.
 	 * @param Tag The tag of the corresponding observation. Must match the tag given during Specify.
 	 * @return The newly created observation object element.
 	 */
-	static UE_API FLearningAgentsObservationObjectElement MakeStructObservationFromArrayViews(ULearningAgentsObservationObject* Object, const TArrayView<const FName> ElementNames, const TArrayView<const FLearningAgentsObservationObjectElement> Elements, const FName Tag = TEXT("StructObservation"));
+	static UE_API FLearningAgentsObservationObjectElement MakeStructObservationFromArrayViews(ULearningAgentsObservationObject* Object, const TArrayView<const FName> ElementNames, const TArrayView<const FLearningAgentsObservationObjectElement> Elements, const bool bActive = true, const FName Tag = TEXT("StructObservation"));
 
 	/**
 	 * Make a new exclusive union observation.
@@ -1082,22 +1151,24 @@ public:
 	 * @param Object The Observation Object
 	 * @param ElementName The name of the chosen sub-observation.
 	 * @param Element The corresponding chosen sub-observation.
+	 * @param bActive Whether the observation should be considered active. Must match the value given during Specify.
 	 * @param Tag The tag of the corresponding observation. Must match the tag given during Specify.
 	 * @return The newly created observation object element.
 	 */
 	UFUNCTION(BlueprintPure, Category = "LearningAgents", meta = (AdvancedDisplay = 3))
-	static UE_API FLearningAgentsObservationObjectElement MakeExclusiveUnionObservation(ULearningAgentsObservationObject* Object, const FName ElementName, const FLearningAgentsObservationObjectElement Element, const FName Tag = TEXT("ExclusiveUnionObservation"));
+	static UE_API FLearningAgentsObservationObjectElement MakeExclusiveUnionObservation(ULearningAgentsObservationObject* Object, const FName ElementName, const FLearningAgentsObservationObjectElement Element, const bool bActive = true, const FName Tag = TEXT("ExclusiveUnionObservation"));
 
 	/**
 	 * Make a new inclusive union observation.
 	 *
 	 * @param Object The Observation Object
 	 * @param Elements The chosen sub-observations.
+	 * @param bActive Whether the observation should be considered active. Must match the value given during Specify.
 	 * @param Tag The tag of the corresponding observation. Must match the tag given during Specify.
 	 * @return The newly created observation object element.
 	 */
 	UFUNCTION(BlueprintPure, Category = "LearningAgents", meta = (AdvancedDisplay = 2, AutoCreateRefTerm = "Elements"))
-	static UE_API FLearningAgentsObservationObjectElement MakeInclusiveUnionObservation(ULearningAgentsObservationObject* Object, const TMap<FName, FLearningAgentsObservationObjectElement>& Elements, const FName Tag = TEXT("InclusiveUnionObservation"));
+	static UE_API FLearningAgentsObservationObjectElement MakeInclusiveUnionObservation(ULearningAgentsObservationObject* Object, const TMap<FName, FLearningAgentsObservationObjectElement>& Elements, const bool bActive = true, const FName Tag = TEXT("InclusiveUnionObservation"));
 
 	/**
 	 * Make a new inclusive union observation.
@@ -1105,11 +1176,12 @@ public:
 	 * @param Object The Observation Object
 	 * @param ElementNames The names of the chosen sub-observations.
 	 * @param Elements The corresponding chosen sub-observations.
+	 * @param bActive Whether the observation should be considered active. Must match the value given during Specify.
 	 * @param Tag The tag of the corresponding observation. Must match the tag given during Specify.
 	 * @return The newly created observation object element.
 	 */
 	UFUNCTION(BlueprintPure, Category = "LearningAgents", meta = (AdvancedDisplay = 3, AutoCreateRefTerm = "ElementNames,Elements"))
-	static UE_API FLearningAgentsObservationObjectElement MakeInclusiveUnionObservationFromArrays(ULearningAgentsObservationObject* Object, const TArray<FName>& ElementNames, const TArray<FLearningAgentsObservationObjectElement>& Elements, const FName Tag = TEXT("InclusiveUnionObservation"));
+	static UE_API FLearningAgentsObservationObjectElement MakeInclusiveUnionObservationFromArrays(ULearningAgentsObservationObject* Object, const TArray<FName>& ElementNames, const TArray<FLearningAgentsObservationObjectElement>& Elements, const bool bActive = true, const FName Tag = TEXT("InclusiveUnionObservation"));
 	
 	/**
 	 * Make a new inclusive union observation.
@@ -1117,63 +1189,69 @@ public:
 	 * @param Object The Observation Object
 	 * @param ElementNames The names of the chosen sub-observations.
 	 * @param Elements The corresponding chosen sub-observations.
+	 * @param bActive Whether the observation should be considered active. Must match the value given during Specify.
 	 * @param Tag The tag of the corresponding observation. Must match the tag given during Specify.
 	 * @return The newly created observation object element.
 	 */
-	static UE_API FLearningAgentsObservationObjectElement MakeInclusiveUnionObservationFromArrayViews(ULearningAgentsObservationObject* Object, const TArrayView<const FName> ElementNames, const TArrayView<const FLearningAgentsObservationObjectElement> Elements, const FName Tag = TEXT("InclusiveUnionObservation"));
+	static UE_API FLearningAgentsObservationObjectElement MakeInclusiveUnionObservationFromArrayViews(ULearningAgentsObservationObject* Object, const TArrayView<const FName> ElementNames, const TArrayView<const FLearningAgentsObservationObjectElement> Elements, const bool bActive = true, const FName Tag = TEXT("InclusiveUnionObservation"));
 
 	/**
 	 * Make a new static array observation.
 	 *
 	 * @param Object The Observation Object
 	 * @param Elements The sub-observations. The number of elements here must match what was given during Specify.
+	 * @param bActive Whether the observation should be considered active. Must match the value given during Specify.
 	 * @param Tag The tag of the corresponding observation. Must match the tag given during Specify.
 	 * @return The newly created observation object element.
 	 */
 	UFUNCTION(BlueprintPure, Category = "LearningAgents", meta = (AdvancedDisplay = 2))
-	static UE_API FLearningAgentsObservationObjectElement MakeStaticArrayObservation(ULearningAgentsObservationObject* Object, const TArray<FLearningAgentsObservationObjectElement>& Elements, const FName Tag = TEXT("StaticArrayObservation"));
+	static UE_API FLearningAgentsObservationObjectElement MakeStaticArrayObservation(ULearningAgentsObservationObject* Object, const TArray<FLearningAgentsObservationObjectElement>& Elements, const bool bActive = true, const FName Tag = TEXT("StaticArrayObservation"));
 
 	/**
 	 * Make a new static array observation.
 	 *
 	 * @param Object The Observation Object
 	 * @param Elements The sub-observations. The number of elements here must match what was given during Specify.
+	 * @param bActive Whether the observation should be considered active. Must match the value given during Specify.
 	 * @param Tag The tag of the corresponding observation. Must match the tag given during Specify.
 	 * @return The newly created observation object element.
 	 */
-	static UE_API FLearningAgentsObservationObjectElement MakeStaticArrayObservationFromArrayView(ULearningAgentsObservationObject* Object, const TArrayView<const FLearningAgentsObservationObjectElement> Elements, const FName Tag = TEXT("StaticArrayObservation"));
+	static UE_API FLearningAgentsObservationObjectElement MakeStaticArrayObservationFromArrayView(ULearningAgentsObservationObject* Object, const TArrayView<const FLearningAgentsObservationObjectElement> Elements, const bool bActive = true, const FName Tag = TEXT("StaticArrayObservation"));
 
 	/**
 	 * Make a new set observation.
 	 *
 	 * @param Object The Observation Object
 	 * @param Elements The sub-observations. The number of elements here must be less than or equal to the maximum that was given during Specify.
+	 * @param bActive Whether the observation should be considered active. Must match the value given during Specify.
 	 * @param Tag The tag of the corresponding observation. Must match the tag given during Specify.
 	 * @return The newly created observation object element.
 	 */
 	UFUNCTION(BlueprintPure, Category = "LearningAgents", meta = (AdvancedDisplay = 2, AutoCreateRefTerm = "Elements"))
-	static UE_API FLearningAgentsObservationObjectElement MakeSetObservation(ULearningAgentsObservationObject* Object, const TSet<FLearningAgentsObservationObjectElement>& Elements, const FName Tag = TEXT("SetObservation"));
+	static UE_API FLearningAgentsObservationObjectElement MakeSetObservation(ULearningAgentsObservationObject* Object, const TSet<FLearningAgentsObservationObjectElement>& Elements, const bool bActive = true, const FName Tag = TEXT("SetObservation"));
 
 	/**
 	 * Make a new set observation.
 	 *
 	 * @param Object The Observation Object
 	 * @param Elements The sub-observations. The number of elements here must be less than or equal to the maximum that was given during Specify.
+	 * @param bActive Whether the observation should be considered active. Must match the value given during Specify.
 	 * @param Tag The tag of the corresponding observation. Must match the tag given during Specify.
 	 * @return The newly created observation object element.
 	 */
 	UFUNCTION(BlueprintPure, Category = "LearningAgents", meta = (AdvancedDisplay = 2, AutoCreateRefTerm = "Elements"))
-	static UE_API FLearningAgentsObservationObjectElement MakeSetObservationFromArray(ULearningAgentsObservationObject* Object, const TArray<FLearningAgentsObservationObjectElement>& Elements, const FName Tag = TEXT("SetObservation"));
+	static UE_API FLearningAgentsObservationObjectElement MakeSetObservationFromArray(ULearningAgentsObservationObject* Object, const TArray<FLearningAgentsObservationObjectElement>& Elements, const bool bActive = true, const FName Tag = TEXT("SetObservation"));
 
 	/**
 	 * Make a new set observation.
 	 *
 	 * @param Object The Observation Object
 	 * @param Elements The sub-observations. The number of elements here must be less than or equal to the maximum that was given during Specify.
+	 * @param bActive Whether the observation should be considered active. Must match the value given during Specify.
 	 * @param Tag The tag of the corresponding observation. Must match the tag given during Specify.
 	 * @return The newly created observation object element.
 	 */
-	static UE_API FLearningAgentsObservationObjectElement MakeSetObservationFromArrayView(ULearningAgentsObservationObject* Object, const TArrayView<const FLearningAgentsObservationObjectElement> Elements, const FName Tag = TEXT("SetObservation"));
+	static UE_API FLearningAgentsObservationObjectElement MakeSetObservationFromArrayView(ULearningAgentsObservationObject* Object, const TArrayView<const FLearningAgentsObservationObjectElement> Elements, const bool bActive = true, const FName Tag = TEXT("SetObservation"));
 
 	/**
 	 * Make a new pair observation.
@@ -1181,11 +1259,12 @@ public:
 	 * @param Object The Observation Object
 	 * @param Key The key sub-observation.
 	 * @param Key The value sub-observation.
+	 * @param bActive Whether the observation should be considered active. Must match the value given during Specify.
 	 * @param Tag The tag of the corresponding observation. Must match the tag given during Specify.
 	 * @return The newly created observation object element.
 	 */
 	UFUNCTION(BlueprintPure, Category = "LearningAgents", meta = (AdvancedDisplay = 3))
-	static UE_API FLearningAgentsObservationObjectElement MakePairObservation(ULearningAgentsObservationObject* Object, const FLearningAgentsObservationObjectElement Key, const FLearningAgentsObservationObjectElement Value, const FName Tag = TEXT("PairObservation"));
+	static UE_API FLearningAgentsObservationObjectElement MakePairObservation(ULearningAgentsObservationObject* Object, const FLearningAgentsObservationObjectElement Key, const FLearningAgentsObservationObjectElement Value, const bool bActive = true, const FName Tag = TEXT("PairObservation"));
 
 	/**
 	 * Make a new array observation.
@@ -1193,11 +1272,12 @@ public:
 	 * @param Object The Observation Object
 	 * @param Elements The sub-observations. The number of elements here must be less than or equal to the maximum that was given during Specify.
 	 * @param MaxNum The maximum number of elements possible for this observation. Must match what was given during Specify.
+	 * @param bActive Whether the observation should be considered active. Must match the value given during Specify.
 	 * @param Tag The tag of the corresponding observation. Must match the tag given during Specify.
 	 * @return The newly created observation object element.
 	 */
 	UFUNCTION(BlueprintPure, Category = "LearningAgents", meta = (AdvancedDisplay = 3, AutoCreateRefTerm = "Elements"))
-	static UE_API FLearningAgentsObservationObjectElement MakeArrayObservation(ULearningAgentsObservationObject* Object, const TArray<FLearningAgentsObservationObjectElement>& Elements, const int32 MaxNum, const FName Tag = TEXT("ArrayObservation"));
+	static UE_API FLearningAgentsObservationObjectElement MakeArrayObservation(ULearningAgentsObservationObject* Object, const TArray<FLearningAgentsObservationObjectElement>& Elements, const int32 MaxNum, const bool bActive = true, const FName Tag = TEXT("ArrayObservation"));
 
 	/**
 	 * Make a new array observation.
@@ -1205,42 +1285,46 @@ public:
 	 * @param Object The Observation Object
 	 * @param Elements The sub-observations. The number of elements here must be less than or equal to the maximum that was given during Specify.
 	 * @param MaxNum The maximum number of elements possible for this observation. Must match what was given during Specify.
+	 * @param bActive Whether the observation should be considered active. Must match the value given during Specify.
 	 * @param Tag The tag of the corresponding observation. Must match the tag given during Specify.
 	 * @return The newly created observation object element.
 	 */
-	static UE_API FLearningAgentsObservationObjectElement MakeArrayObservationFromArrayView(ULearningAgentsObservationObject* Object, const TArrayView<const FLearningAgentsObservationObjectElement> Elements, const int32 MaxNum, const FName Tag = TEXT("ArrayObservation"));
+	static UE_API FLearningAgentsObservationObjectElement MakeArrayObservationFromArrayView(ULearningAgentsObservationObject* Object, const TArrayView<const FLearningAgentsObservationObjectElement> Elements, const int32 MaxNum, const bool bActive = true, const FName Tag = TEXT("ArrayObservation"));
 
 	/**
 	 * Make a new map observation.
 	 *
 	 * @param Object The Observation Object
 	 * @param Elements The sub-observations. The number of elements here must be less than or equal to the maximum that was given during Specify.
+	 * @param bActive Whether the observation should be considered active. Must match the value given during Specify.
 	 * @param Tag The tag of the corresponding observation. Must match the tag given during Specify.
 	 * @return The newly created observation object element.
 	 */
 	UFUNCTION(BlueprintPure, Category = "LearningAgents", meta = (AdvancedDisplay = 2, AutoCreateRefTerm = "Map"))
-	static UE_API FLearningAgentsObservationObjectElement MakeMapObservation(ULearningAgentsObservationObject* Object, const TMap<FLearningAgentsObservationObjectElement, FLearningAgentsObservationObjectElement>& Map, const FName Tag = TEXT("MapObservation"));
+	static UE_API FLearningAgentsObservationObjectElement MakeMapObservation(ULearningAgentsObservationObject* Object, const TMap<FLearningAgentsObservationObjectElement, FLearningAgentsObservationObjectElement>& Map, const bool bActive = true, const FName Tag = TEXT("MapObservation"));
 
 	/**
 	 * Make a new map observation.
 	 *
 	 * @param Object The Observation Object
 	 * @param Elements The sub-observations. The number of elements here must be less than or equal to the maximum that was given during Specify.
+	 * @param bActive Whether the observation should be considered active. Must match the value given during Specify.
 	 * @param Tag The tag of the corresponding observation. Must match the tag given during Specify.
 	 * @return The newly created observation object element.
 	 */
 	UFUNCTION(BlueprintPure, Category = "LearningAgents", meta = (AdvancedDisplay = 3, AutoCreateRefTerm = "Keys,Values"))
-	static UE_API FLearningAgentsObservationObjectElement MakeMapObservationFromArrays(ULearningAgentsObservationObject* Object, const TArray<FLearningAgentsObservationObjectElement>& Keys, const TArray<FLearningAgentsObservationObjectElement>& Values, const FName Tag = TEXT("MapObservation"));
+	static UE_API FLearningAgentsObservationObjectElement MakeMapObservationFromArrays(ULearningAgentsObservationObject* Object, const TArray<FLearningAgentsObservationObjectElement>& Keys, const TArray<FLearningAgentsObservationObjectElement>& Values, const bool bActive = true, const FName Tag = TEXT("MapObservation"));
 	
 	/**
 	 * Make a new map observation.
 	 *
 	 * @param Object The Observation Object
 	 * @param Elements The sub-observations. The number of elements here must be less than or equal to the maximum that was given during Specify.
+	 * @param bActive Whether the observation should be considered active. Must match the value given during Specify.
 	 * @param Tag The tag of the corresponding observation. Must match the tag given during Specify.
 	 * @return The newly created observation object element.
 	 */
-	static UE_API FLearningAgentsObservationObjectElement MakeMapObservationFromArrayViews(ULearningAgentsObservationObject* Object, const TArrayView<const FLearningAgentsObservationObjectElement> Keys, const TArrayView<const FLearningAgentsObservationObjectElement> Values, const FName Tag = TEXT("MapObservation"));
+	static UE_API FLearningAgentsObservationObjectElement MakeMapObservationFromArrayViews(ULearningAgentsObservationObject* Object, const TArrayView<const FLearningAgentsObservationObjectElement> Keys, const TArrayView<const FLearningAgentsObservationObjectElement> Values, const bool bActive = true, const FName Tag = TEXT("MapObservation"));
 
 	/**
 	 * Make a new enum observation.
@@ -1248,6 +1332,7 @@ public:
 	 * @param Object The Observation Object
 	 * @param Enum The enum type for this observation. Must match what was given during Specify.
 	 * @param EnumValue The enum value.
+	 * @param bActive Whether the observation should be considered active. Must match the value given during Specify.
 	 * @param Tag The tag of the corresponding observation. Must match the tag given during Specify.
 	 * @param bVisualLoggerEnabled When true, debug data will be sent to the visual logger.
 	 * @param VisualLoggerListener The listener object which is making this observation. This must be set to use logging.
@@ -1261,6 +1346,7 @@ public:
 		ULearningAgentsObservationObject* Object, 
 		const UEnum* Enum, 
 		const uint8 EnumValue, 
+		const bool bActive = true, 
 		const FName Tag = TEXT("EnumObservation"),
 		const bool bVisualLoggerEnabled = false,
 		ULearningAgentsManagerListener* VisualLoggerListener = nullptr,
@@ -1274,6 +1360,7 @@ public:
 	 * @param Object The Observation Object
 	 * @param Enum The enum type for this observation. Must match what was given during Specify.
 	 * @param BitmaskValue The bitmask value.
+	 * @param bActive Whether the observation should be considered active. Must match the value given during Specify.
 	 * @param Tag The tag of the corresponding observation. Must match the tag given during Specify.
 	 * @param bVisualLoggerEnabled When true, debug data will be sent to the visual logger.
 	 * @param VisualLoggerListener The listener object which is making this observation. This must be set to use logging.
@@ -1287,6 +1374,7 @@ public:
 		ULearningAgentsObservationObject* Object, 
 		const UEnum* Enum, 
 		const int32 BitmaskValue, 
+		const bool bActive = true, 
 		const FName Tag = TEXT("BitmaskObservation"),
 		const bool bVisualLoggerEnabled = false,
 		ULearningAgentsManagerListener* VisualLoggerListener = nullptr,
@@ -1300,31 +1388,35 @@ public:
 	 * @param Object The Observation Object
 	 * @param Element The sub-observation given.
 	 * @param Option The indicator as to if this is observation should be used.
+	 * @param bActive Whether the observation should be considered active. Must match the value given during Specify.
 	 * @param Tag The tag of the corresponding observation. Must match the tag given during Specify.
 	 * @return The newly created observation object element.
 	 */
 	UFUNCTION(BlueprintPure, Category = "LearningAgents", meta = (AdvancedDisplay = 3))
-	static UE_API FLearningAgentsObservationObjectElement MakeOptionalObservation(ULearningAgentsObservationObject* Object, const FLearningAgentsObservationObjectElement Element, const ELearningAgentsOptionalObservation Option, const FName Tag = TEXT("OptionalObservation"));
+	static UE_API FLearningAgentsObservationObjectElement MakeOptionalObservation(ULearningAgentsObservationObject* Object, const FLearningAgentsObservationObjectElement Element, const ELearningAgentsOptionalObservation Option, const bool bActive = true, const FName Tag = TEXT("OptionalObservation"));
 
 	/**
 	 * Make a new null optional observation. Use this to provide a null optional observation.
 	 *
 	 * @param Object The Observation Object
+	 * @param bActive Whether the observation should be considered active. Must match the value given during Specify.
 	 * @param Tag The tag of the corresponding observation. Must match the tag given during Specify.
 	 * @return The newly created observation object element.
 	 */
 	UFUNCTION(BlueprintPure, Category = "LearningAgents", meta = (AdvancedDisplay = 1))
-	static UE_API FLearningAgentsObservationObjectElement MakeOptionalNullObservation(ULearningAgentsObservationObject* Object, const FName Tag = TEXT("OptionalObservation"));
+	static UE_API FLearningAgentsObservationObjectElement MakeOptionalNullObservation(ULearningAgentsObservationObject* Object, const bool bActive = true, const FName Tag = TEXT("OptionalObservation"));
 
 	/**
 	 * Make a new valid optional observation. Use this to provide a valid optional observation.
 	 *
 	 * @param Object The Observation Object
+	 * @param Element The sub-observation given.
+	 * @param bActive Whether the observation should be considered active. Must match the value given during Specify.
 	 * @param Tag The tag of the corresponding observation. Must match the tag given during Specify.
 	 * @return The newly created observation object element.
 	 */
 	UFUNCTION(BlueprintPure, Category = "LearningAgents", meta = (AdvancedDisplay = 2))
-	static UE_API FLearningAgentsObservationObjectElement MakeOptionalValidObservation(ULearningAgentsObservationObject* Object, const FLearningAgentsObservationObjectElement Element, const FName Tag = TEXT("OptionalObservation"));
+	static UE_API FLearningAgentsObservationObjectElement MakeOptionalValidObservation(ULearningAgentsObservationObject* Object, const FLearningAgentsObservationObjectElement Element, const bool bActive = true, const FName Tag = TEXT("OptionalObservation"));
 
 	/**
 	 * Make a new either observation.
@@ -1332,50 +1424,55 @@ public:
 	 * @param Object The Observation Object
 	 * @param Element The sub-observation given.
 	 * @param Option The indicator as to if this is observation A or B.
+	 * @param bActive Whether the observation should be considered active. Must match the value given during Specify.
 	 * @param Tag The tag of the corresponding observation. Must match the tag given during Specify.
 	 * @return The newly created observation object element.
 	 */
 	UFUNCTION(BlueprintPure, Category = "LearningAgents", meta = (AdvancedDisplay = 3))
-	static UE_API FLearningAgentsObservationObjectElement MakeEitherObservation(ULearningAgentsObservationObject* Object, const FLearningAgentsObservationObjectElement Element, const ELearningAgentsEitherObservation Either, const FName Tag = TEXT("EitherObservation"));
+	static UE_API FLearningAgentsObservationObjectElement MakeEitherObservation(ULearningAgentsObservationObject* Object, const FLearningAgentsObservationObjectElement Element, const ELearningAgentsEitherObservation Either, const bool bActive = true, const FName Tag = TEXT("EitherObservation"));
 
 	/**
 	 * Make a new either A observation. Use this to provide option A.
 	 *
 	 * @param Object The Observation Object
 	 * @param Element The sub-observation given.
+	 * @param bActive Whether the observation should be considered active. Must match the value given during Specify.
 	 * @param Tag The tag of the corresponding observation. Must match the tag given during Specify.
 	 * @return The newly created observation object element.
 	 */
 	UFUNCTION(BlueprintPure, Category = "LearningAgents", meta = (AdvancedDisplay = 2, DisplayName = "Make Either A Observation"))
-	static UE_API FLearningAgentsObservationObjectElement MakeEitherAObservation(ULearningAgentsObservationObject* Object, const FLearningAgentsObservationObjectElement A, const FName Tag = TEXT("EitherObservation"));
+	static UE_API FLearningAgentsObservationObjectElement MakeEitherAObservation(ULearningAgentsObservationObject* Object, const FLearningAgentsObservationObjectElement A, const bool bActive = true, const FName Tag = TEXT("EitherObservation"));
 
 	/**
 	 * Make a new either B observation. Use this to provide option B.
 	 *
 	 * @param Object The Observation Object
 	 * @param Element The sub-observation given.
+	 * @param bActive Whether the observation should be considered active. Must match the value given during Specify.
 	 * @param Tag The tag of the corresponding observation. Must match the tag given during Specify.
 	 * @return The newly created observation object element.
 	 */
 	UFUNCTION(BlueprintPure, Category = "LearningAgents", meta = (AdvancedDisplay = 2, DisplayName = "Make Either B Observation"))
-	static UE_API FLearningAgentsObservationObjectElement MakeEitherBObservation(ULearningAgentsObservationObject* Object, const FLearningAgentsObservationObjectElement B, const FName Tag = TEXT("EitherObservation"));
+	static UE_API FLearningAgentsObservationObjectElement MakeEitherBObservation(ULearningAgentsObservationObject* Object, const FLearningAgentsObservationObjectElement B, const bool bActive = true, const FName Tag = TEXT("EitherObservation"));
 
 	/**
 	 * Make a new encoding observation. This must be used in conjunction with SpecifyEncodingObservation.
-	 * 
+	 *
 	 * @param Object The Observation Object
 	 * @param Element The Observation Element to be encoded.
+	 * @param bActive Whether the observation should be considered active. Must match the value given during Specify.
 	 * @param Tag The tag of the corresponding observation. Must match the tag given during Specify.
 	 * @return The newly created observation object element.
 	 */
 	UFUNCTION(BlueprintPure, Category = "LearningAgents", meta = (AdvancedDisplay = 2))
-	static UE_API FLearningAgentsObservationObjectElement MakeEncodingObservation(ULearningAgentsObservationObject* Object, const FLearningAgentsObservationObjectElement Element, const FName Tag = TEXT("EncodingObservation"));
+	static UE_API FLearningAgentsObservationObjectElement MakeEncodingObservation(ULearningAgentsObservationObject* Object, const FLearningAgentsObservationObjectElement Element, const bool bActive = true, const FName Tag = TEXT("EncodingObservation"));
 
 	/**
 	 * Make a new bool observation.
 	 *
 	 * @param Object The Observation Object
 	 * @param Value The new value of this observation.
+	 * @param bActive Whether the observation should be considered active. Must match the value given during Specify.
 	 * @param Tag The tag of the corresponding observation. Must match the tag given during Specify.
 	 * @param bVisualLoggerEnabled When true, debug data will be sent to the visual logger.
 	 * @param VisualLoggerListener The listener object which is making this observation. This must be set to use logging.
@@ -1388,6 +1485,7 @@ public:
 	static UE_API FLearningAgentsObservationObjectElement MakeBoolObservation(
 		ULearningAgentsObservationObject* Object,
 		const bool bValue,
+		const bool bActive = true, 
 		const FName Tag = TEXT("BoolObservation"),
 		const bool bVisualLoggerEnabled = false,
 		ULearningAgentsManagerListener* VisualLoggerListener = nullptr,
@@ -1400,26 +1498,31 @@ public:
 	 *
 	 * @param Object The Observation Object
 	 * @param Element The Observation Element to be encoded.
+	 * @param bActive Whether the observation should be considered active. Must match the value given during Specify.
+	 * @param Tag The tag of the corresponding observation. Must match the tag given during Specify.
 	 * @return The newly created observation object element.
 	 */
 	UFUNCTION(BlueprintPure, Category = "LearningAgents", meta = (AdvancedDisplay = 2, DefaultToSelf = "VisualLoggerListener"))
-	static UE_API FLearningAgentsObservationObjectElement MakeConv1dObservation(ULearningAgentsObservationObject* Object, const FLearningAgentsObservationObjectElement Element, const FName Tag = TEXT("Conv1dObservation"));
+	static UE_API FLearningAgentsObservationObjectElement MakeConv1dObservation(ULearningAgentsObservationObject* Object, const FLearningAgentsObservationObjectElement Element, const bool bActive = true, const FName Tag = TEXT("Conv1dObservation"));
 
 	/**
 	 * Make a new 2d convolution observation.
 	 *
 	 * @param Object The Observation Object
 	 * @param Element The Observation Element to be encoded.
+	 * @param bActive Whether the observation should be considered active. Must match the value given during Specify.
+	 * @param Tag The tag of the corresponding observation. Must match the tag given during Specify.
 	 * @return The newly created observation object element.
 	 */
 	UFUNCTION(BlueprintPure, Category = "LearningAgents", meta = (AdvancedDisplay = 2, DefaultToSelf = "VisualLoggerListener"))
-	static UE_API FLearningAgentsObservationObjectElement MakeConv2dObservation(ULearningAgentsObservationObject* Object, const FLearningAgentsObservationObjectElement Element, const FName Tag = TEXT("Conv2dObservation"));
+	static UE_API FLearningAgentsObservationObjectElement MakeConv2dObservation(ULearningAgentsObservationObject* Object, const FLearningAgentsObservationObjectElement Element, const bool bActive = true, const FName Tag = TEXT("Conv2dObservation"));
 
 	/**
 	 * Make a new float observation.
 	 *
 	 * @param Object The Observation Object
 	 * @param Value The new value of this observation.
+	 * @param bActive Whether the observation should be considered active. Must match the value given during Specify.
 	 * @param Tag The tag of the corresponding observation. Must match the tag given during Specify.
 	 * @param bVisualLoggerEnabled When true, debug data will be sent to the visual logger.
 	 * @param VisualLoggerListener The listener object which is making this observation. This must be set to use logging.
@@ -1432,6 +1535,7 @@ public:
 	static UE_API FLearningAgentsObservationObjectElement MakeFloatObservation(
 		ULearningAgentsObservationObject* Object,
 		const float Value,
+		const bool bActive = true, 
 		const FName Tag = TEXT("FloatObservation"),
 		const bool bVisualLoggerEnabled = false,
 		ULearningAgentsManagerListener* VisualLoggerListener = nullptr,
@@ -1445,6 +1549,7 @@ public:
 	 * @param Object The Observation Object
 	 * @param Location The location of interest to the agent.
 	 * @param RelativeTransform The transform the provided location should be encoded relative to.
+	 * @param bActive Whether the observation should be considered active. Must match the value given during Specify.
 	 * @param Tag The tag of the corresponding observation. Must match the tag given during Specify.
 	 * @param bVisualLoggerEnabled When true, debug data will be sent to the visual logger.
 	 * @param VisualLoggerListener The listener object which is making this observation. This must be set to use logging.
@@ -1458,6 +1563,7 @@ public:
 		ULearningAgentsObservationObject* Object,
 		const FVector Location,
 		const FTransform RelativeTransform = FTransform(),
+		const bool bActive = true, 
 		const FName Tag = TEXT("LocationObservation"),
 		const bool bVisualLoggerEnabled = false,
 		ULearningAgentsManagerListener* VisualLoggerListener = nullptr,
@@ -1471,6 +1577,7 @@ public:
 	 * @param Object The Observation Object
 	 * @param Rotation The rotation of interest to the agent.
 	 * @param RelativeRotation The rotation the provided rotation should be encoded relative to.
+	 * @param bActive Whether the observation should be considered active. Must match the value given during Specify.
 	 * @param Tag The tag of the corresponding observation. Must match the tag given during Specify.
 	 * @param bVisualLoggerEnabled When true, debug data will be sent to the visual logger.
 	 * @param VisualLoggerListener The listener object which is making this observation. This must be set to use logging.
@@ -1485,6 +1592,7 @@ public:
 		ULearningAgentsObservationObject* Object,
 		const FRotator Rotation,
 		const FRotator RelativeRotation = FRotator::ZeroRotator,
+		const bool bActive = true, 
 		const FName Tag = TEXT("RotationObservation"),
 		const bool bVisualLoggerEnabled = false,
 		ULearningAgentsManagerListener* VisualLoggerListener = nullptr,
@@ -1499,6 +1607,7 @@ public:
 	 * @param Object The Observation Object
 	 * @param Rotation The rotation of interest to the agent.
 	 * @param RelativeRotation The rotation the provided rotation should be encoded relative to.
+	 * @param bActive Whether the observation should be considered active. Must match the value given during Specify.
 	 * @param Tag The tag of the corresponding observation. Must match the tag given during Specify.
 	 * @param bVisualLoggerEnabled When true, debug data will be sent to the visual logger.
 	 * @param VisualLoggerListener The listener object which is making this observation. This must be set to use logging.
@@ -1513,6 +1622,7 @@ public:
 		ULearningAgentsObservationObject* Object,
 		const FQuat Rotation,
 		const FQuat RelativeRotation,
+		const bool bActive = true, 
 		const FName Tag = TEXT("RotationObservation"),
 		const bool bVisualLoggerEnabled = false,
 		ULearningAgentsManagerListener* VisualLoggerListener = nullptr,
@@ -1527,6 +1637,7 @@ public:
 	 * @param Object The Observation Object
 	 * @param Scale The scale of interest to the agent.
 	 * @param RelativeScale The scale the provided scale should be encoded relative to.
+	 * @param bActive Whether the observation should be considered active. Must match the value given during Specify.
 	 * @param Tag The tag of the corresponding observation. Must match the tag given during Specify.
 	 * @param bVisualLoggerEnabled When true, debug data will be sent to the visual logger.
 	 * @param VisualLoggerListener The listener object which is making this observation. This must be set to use logging.
@@ -1541,6 +1652,7 @@ public:
 		ULearningAgentsObservationObject* Object,
 		const FVector Scale,
 		const FVector RelativeScale = FVector(1, 1, 1),
+		const bool bActive = true, 
 		const FName Tag = TEXT("ScaleObservation"),
 		const bool bVisualLoggerEnabled = false,
 		ULearningAgentsManagerListener* VisualLoggerListener = nullptr,
@@ -1555,6 +1667,7 @@ public:
 	 * @param Object The Observation Object
 	 * @param Transform The transform of interest to the agent.
 	 * @param RelativeTransform The transform the provided transform should be encoded relative to.
+	 * @param bActive Whether the observation should be considered active. Must match the value given during Specify.
 	 * @param Tag The tag of the corresponding observation. Must match the tag given during Specify.
 	 * @param bVisualLoggerEnabled When true, debug data will be sent to the visual logger.
 	 * @param VisualLoggerListener The listener object which is making this observation. This must be set to use logging.
@@ -1568,6 +1681,7 @@ public:
 		ULearningAgentsObservationObject* Object,
 		const FTransform Transform,
 		const FTransform RelativeTransform = FTransform(),
+		const bool bActive = true, 
 		const FName Tag = TEXT("TransformObservation"),
 		const bool bVisualLoggerEnabled = false,
 		ULearningAgentsManagerListener* VisualLoggerListener = nullptr,
@@ -1581,6 +1695,7 @@ public:
 	 * @param Object The Observation Object
 	 * @param Angle The angle of interest to the agent.
 	 * @param RelativeAngle The angle the provided angle should be encoded relative to.
+	 * @param bActive Whether the observation should be considered active. Must match the value given during Specify.
 	 * @param Tag The tag of the corresponding observation. Must match the tag given during Specify.
 	 * @param bVisualLoggerEnabled When true, debug data will be sent to the visual logger.
 	 * @param VisualLoggerListener The listener object which is making this observation. This must be set to use logging.
@@ -1591,9 +1706,10 @@ public:
 	 */
 	UFUNCTION(BlueprintPure, Category = "LearningAgents", meta = (AdvancedDisplay = 3, DefaultToSelf = "VisualLoggerListener"))
 	static UE_API FLearningAgentsObservationObjectElement MakeAngleObservation(
-		ULearningAgentsObservationObject* Object, 
-		const float Angle, 
-		const float RelativeAngle = 0.0f, 
+		ULearningAgentsObservationObject* Object,
+		const float Angle,
+		const float RelativeAngle = 0.0f,
+		const bool bActive = true,
 		const FName Tag = TEXT("AngleObservation"),
 		const bool bVisualLoggerEnabled = false,
 		ULearningAgentsManagerListener* VisualLoggerListener = nullptr,
@@ -1603,10 +1719,11 @@ public:
 
 	/**
 	 * Make a new angle observation. Angles should be given in radians.
-	 * 
+	 *
 	 * @param Object The Observation Object
 	 * @param Angle The angle of interest to the agent.
 	 * @param RelativeAngle The angle the provided angle should be encoded relative to.
+	 * @param bActive Whether the observation should be considered active. Must match the value given during Specify.
 	 * @param Tag The tag of the corresponding observation. Must match the tag given during Specify.
 	 * @param bVisualLoggerEnabled When true, debug data will be sent to the visual logger.
 	 * @param VisualLoggerListener The listener object which is making this observation. This must be set to use logging.
@@ -1620,6 +1737,7 @@ public:
 		ULearningAgentsObservationObject* Object, 
 		const float Angle, 
 		const float RelativeAngle = 0.0f, 
+		const bool bActive = true, 
 		const FName Tag = TEXT("AngleObservation"),
 		const bool bVisualLoggerEnabled = false,
 		ULearningAgentsManagerListener* VisualLoggerListener = nullptr,
@@ -1633,6 +1751,7 @@ public:
 	 * @param Object The Observation Object
 	 * @param Velocity The velocity of interest to the agent.
 	 * @param RelativeTransform The transform the provided velocity should be encoded relative to.
+	 * @param bActive Whether the observation should be considered active. Must match the value given during Specify.
 	 * @param Tag The tag of the corresponding observation. Must match the tag given during Specify.
 	 * @param bVisualLoggerEnabled When true, debug data will be sent to the visual logger.
 	 * @param VisualLoggerListener The listener object which is making this observation. This must be set to use logging.
@@ -1647,6 +1766,7 @@ public:
 		ULearningAgentsObservationObject* Object, 
 		const FVector Velocity, 
 		const FTransform RelativeTransform = FTransform(), 
+		const bool bActive = true, 
 		const FName Tag = TEXT("VelocityObservation"),
 		const bool bVisualLoggerEnabled = false,
 		ULearningAgentsManagerListener* VisualLoggerListener = nullptr,
@@ -1661,6 +1781,7 @@ public:
 	 * @param Object The Observation Object
 	 * @param Direction The direction of interest to the agent.
 	 * @param RelativeTransform The transform the provided direction should be encoded relative to.
+	 * @param bActive Whether the observation should be considered active. Must match the value given during Specify.
 	 * @param Tag The tag of the corresponding observation. Must match the tag given during Specify.
 	 * @param bVisualLoggerEnabled When true, debug data will be sent to the visual logger.
 	 * @param VisualLoggerListener The listener object which is making this observation. This must be set to use logging.
@@ -1676,6 +1797,7 @@ public:
 		ULearningAgentsObservationObject* Object,
 		const FVector Direction,
 		const FTransform RelativeTransform = FTransform(),
+		const bool bActive = true, 
 		const FName Tag = TEXT("DirectionObservation"),
 		const bool bVisualLoggerEnabled = false,
 		ULearningAgentsManagerListener* VisualLoggerListener = nullptr,
@@ -1692,6 +1814,7 @@ public:
 	 * @param SplineComponent The spline to observe.
 	 * @param DistanceAlongSpline The distance along that spline.
 	 * @param RelativeTransform The transform the provided location should be encoded relative to.
+	 * @param bActive Whether the observation should be considered active. Must match the value given during Specify.
 	 * @param Tag The tag of the corresponding observation. Must match the tag given during Specify.
 	 * @param bVisualLoggerEnabled When true, debug data will be sent to the visual logger.
 	 * @param VisualLoggerListener The listener object which is making this observation. This must be set to use logging.
@@ -1706,6 +1829,7 @@ public:
 		const USplineComponent* SplineComponent,
 		const float DistanceAlongSpline,
 		const FTransform RelativeTransform = FTransform(),
+		const bool bActive = true, 
 		const FName Tag = TEXT("LocationAlongSplineObservation"),
 		const bool bVisualLoggerEnabled = false,
 		ULearningAgentsManagerListener* VisualLoggerListener = nullptr,
@@ -1719,6 +1843,7 @@ public:
 	 * @param Object The Observation Object
 	 * @param SplineComponent The spline to observe.
 	 * @param DistanceAlongSpline The distance along that spline.
+	 * @param bActive Whether the observation should be considered active. Must match the value given during Specify.
 	 * @param Tag The tag of the corresponding observation. Must match the tag given during Specify.
 	 * @param bVisualLoggerEnabled When true, debug data will be sent to the visual logger.
 	 * @param VisualLoggerListener The listener object which is making this observation. This must be set to use logging.
@@ -1732,6 +1857,7 @@ public:
 		ULearningAgentsObservationObject* Object, 
 		const USplineComponent* SplineComponent, 
 		const float DistanceAlongSpline, 
+		const bool bActive = true, 
 		const FName Tag = TEXT("ProportionAlongSplineObservation"),
 		const bool bVisualLoggerEnabled = false,
 		ULearningAgentsManagerListener* VisualLoggerListener = nullptr,
@@ -1746,6 +1872,7 @@ public:
 	 * @param SplineComponent The spline to observe.
 	 * @param DistanceAlongSpline The distance along that spline.
 	 * @param RelativeTransform The transform the provided direction should be encoded relative to.
+	 * @param bActive Whether the observation should be considered active. Must match the value given during Specify.
 	 * @param Tag The tag of the corresponding observation. Must match the tag given during Specify.
 	 * @param bVisualLoggerEnabled When true, debug data will be sent to the visual logger.
 	 * @param VisualLoggerListener The listener object which is making this observation. This must be set to use logging.
@@ -1761,6 +1888,7 @@ public:
 		const USplineComponent* SplineComponent,
 		const float DistanceAlongSpline,
 		const FTransform RelativeTransform = FTransform(),
+		const bool bActive = true, 
 		const FName Tag = TEXT("DirectionAlongSplineObservation"),
 		const bool bVisualLoggerEnabled = false,
 		ULearningAgentsManagerListener* VisualLoggerListener = nullptr,
@@ -1770,13 +1898,14 @@ public:
 		const FLinearColor VisualLoggerColor = FLinearColor::Red);
 
 	/**
-	 * Make a new proportion along ray observation. 
+	 * Make a new proportion along ray observation.
 	 *
 	 * @param Object The Observation Object
 	 * @param RayStart The local ray start location.
 	 * @param RayEnd The local ray end location.
 	 * @param RayTransform The transform to use to transform the local ray starts and ends into the world space.
 	 * @param CollisionChannel The collision channel to collide against.
+	 * @param bActive Whether the observation should be considered active. Must match the value given during Specify.
 	 * @param Tag The tag of the corresponding observation. Must match the tag given during Specify.
 	 * @param bVisualLoggerEnabled When true, debug data will be sent to the visual logger.
 	 * @param VisualLoggerListener The listener object which is making this observation. This must be set to use logging.
@@ -1792,6 +1921,7 @@ public:
 		const FVector RayEnd, 
 		const FTransform RayTransform = FTransform(), 
 		const ECollisionChannel CollisionChannel = ECollisionChannel::ECC_WorldStatic, 
+		const bool bActive = true, 
 		const FName Tag = TEXT("ProportionAlongRayObservation"),
 		const bool bVisualLoggerEnabled = false,
 		ULearningAgentsManagerListener* VisualLoggerListener = nullptr,
@@ -1800,7 +1930,7 @@ public:
 		const FLinearColor VisualLoggerColor = FLinearColor::Red);
 
 	/**
-	 * Make a new LiDAR observation. 
+	 * Make a new LiDAR observation.
 	 *
 	 * @param Object The Observation Object
 	 * @param RayStart The local ray start location.
@@ -1808,6 +1938,7 @@ public:
 	 * @param RayTransform The transform to use to transform the local ray starts and ends into the world space.
 	 * @param CollisionChannel The collision channel to collide against.
 	 * @param ActorToIgnore The actor to ignore in the raycast.
+	 * @param bActive Whether the observation should be considered active. Must match the value given during Specify.
 	 * @param Tag The tag of the corresponding observation. Must match the tag given during Specify.
 	 * @param bVisualLoggerEnabled When true, debug data will be sent to the visual logger.
 	 * @param VisualLoggerListener The listener object which is making this observation. This must be set to use logging.
@@ -1824,6 +1955,7 @@ public:
 		const FTransform RayTransform = FTransform(),
 		const ECollisionChannel CollisionChannel = ECollisionChannel::ECC_WorldStatic,
 		const AActor* ActorToIgnore = nullptr,
+		const bool bActive = true, 
 		const FName Tag = TEXT("LidarObservation"),
 		const bool bVisualLoggerEnabled = false,
 		ULearningAgentsManagerListener* VisualLoggerListener = nullptr,
@@ -1839,167 +1971,167 @@ public:
 	static UE_API bool GetNullObservation(const ULearningAgentsObservationObject* Object, const FLearningAgentsObservationObjectElement Element, const FName Tag = TEXT("NullObservation"));
 
 	UFUNCTION(BlueprintPure = false, Category = "LearningAgents", meta = (AdvancedDisplay = 3, ReturnDisplayName = "Success"))
-	static UE_API bool GetContinuousObservationNum(int32& OutNum, const ULearningAgentsObservationObject* Object, const FLearningAgentsObservationObjectElement Element, const FName Tag = TEXT("ContinuousObservation"));
+	static UE_API bool GetContinuousObservationNum(int32& OutNum, const ULearningAgentsObservationObject* Object, const FLearningAgentsObservationObjectElement Element, const bool bActive = true, const FName Tag = TEXT("ContinuousObservation"));
 
 	UFUNCTION(BlueprintPure = false, Category = "LearningAgents", meta = (AdvancedDisplay = 3, ReturnDisplayName = "Success"))
-	static UE_API bool GetContinuousObservation(TArray<float>& OutValues, const ULearningAgentsObservationObject* Object, const FLearningAgentsObservationObjectElement Element, const FName Tag = TEXT("ContinuousObservation"));
-	static UE_API bool GetContinuousObservationToArrayView(TArrayView<float> OutValues, const ULearningAgentsObservationObject* Object, const FLearningAgentsObservationObjectElement Element, const FName Tag = TEXT("ContinuousObservation"));
+	static UE_API bool GetContinuousObservation(TArray<float>& OutValues, const ULearningAgentsObservationObject* Object, const FLearningAgentsObservationObjectElement Element, const bool bActive = true, const FName Tag = TEXT("ContinuousObservation"));
+	static UE_API bool GetContinuousObservationToArrayView(TArrayView<float> OutValues, const ULearningAgentsObservationObject* Object, const FLearningAgentsObservationObjectElement Element, const bool bActive = true, const FName Tag = TEXT("ContinuousObservation"));
 
 	UFUNCTION(BlueprintPure = false, Category = "LearningAgents", meta = (AdvancedDisplay = 3, ReturnDisplayName = "Success"))
-	static UE_API bool GetExclusiveDiscreteObservation(int32& OutIndex, const ULearningAgentsObservationObject* Object, const FLearningAgentsObservationObjectElement Element, const FName Tag = TEXT("ExclusiveDiscreteObservation"));
+	static UE_API bool GetExclusiveDiscreteObservation(int32& OutIndex, const ULearningAgentsObservationObject* Object, const FLearningAgentsObservationObjectElement Element, const bool bActive = true, const FName Tag = TEXT("ExclusiveDiscreteObservation"));
 
 	UFUNCTION(BlueprintPure = false, Category = "LearningAgents", meta = (AdvancedDisplay = 3, ReturnDisplayName = "Success"))
-	static UE_API bool GetInclusiveDiscreteObservationNum(int32& OutNum, const ULearningAgentsObservationObject* Object, const FLearningAgentsObservationObjectElement Element, const FName Tag = TEXT("InclusiveDiscreteObservation"));
+	static UE_API bool GetInclusiveDiscreteObservationNum(int32& OutNum, const ULearningAgentsObservationObject* Object, const FLearningAgentsObservationObjectElement Element, const bool bActive = true, const FName Tag = TEXT("InclusiveDiscreteObservation"));
 
 	UFUNCTION(BlueprintPure = false, Category = "LearningAgents", meta = (AdvancedDisplay = 3, ReturnDisplayName = "Success"))
-	static UE_API bool GetInclusiveDiscreteObservation(TArray<int32>& OutIndices, const ULearningAgentsObservationObject* Object, const FLearningAgentsObservationObjectElement Element, const FName Tag = TEXT("InclusiveDiscreteObservation"));
-	static UE_API bool GetInclusiveDiscreteObservationToArrayView(TArrayView<int32> OutIndices, const ULearningAgentsObservationObject* Object, const FLearningAgentsObservationObjectElement Element, const FName Tag = TEXT("InclusiveDiscreteObservation"));
+	static UE_API bool GetInclusiveDiscreteObservation(TArray<int32>& OutIndices, const ULearningAgentsObservationObject* Object, const FLearningAgentsObservationObjectElement Element, const bool bActive = true, const FName Tag = TEXT("InclusiveDiscreteObservation"));
+	static UE_API bool GetInclusiveDiscreteObservationToArrayView(TArrayView<int32> OutIndices, const ULearningAgentsObservationObject* Object, const FLearningAgentsObservationObjectElement Element, const bool bActive = true, const FName Tag = TEXT("InclusiveDiscreteObservation"));
 
 	UFUNCTION(BlueprintPure = false, Category = "LearningAgents", meta = (AdvancedDisplay = 3, ReturnDisplayName = "Success"))
-	static UE_API bool GetNamedExclusiveDiscreteObservation(FName& OutName, const ULearningAgentsObservationObject* Object, const FLearningAgentsObservationObjectElement Element, const FName Tag = TEXT("ExclusiveDiscreteObservation"));
+	static UE_API bool GetNamedExclusiveDiscreteObservation(FName& OutName, const ULearningAgentsObservationObject* Object, const FLearningAgentsObservationObjectElement Element, const bool bActive = true, const FName Tag = TEXT("ExclusiveDiscreteObservation"));
 
 	UFUNCTION(BlueprintPure = false, Category = "LearningAgents", meta = (AdvancedDisplay = 3, ReturnDisplayName = "Success"))
-	static UE_API bool GetNamedInclusiveDiscreteObservationNum(int32& OutNum, const ULearningAgentsObservationObject* Object, const FLearningAgentsObservationObjectElement Element, const FName Tag = TEXT("InclusiveDiscreteObservation"));
+	static UE_API bool GetNamedInclusiveDiscreteObservationNum(int32& OutNum, const ULearningAgentsObservationObject* Object, const FLearningAgentsObservationObjectElement Element, const bool bActive = true, const FName Tag = TEXT("InclusiveDiscreteObservation"));
 
 	UFUNCTION(BlueprintPure = false, Category = "LearningAgents", meta = (AdvancedDisplay = 3, ReturnDisplayName = "Success"))
-	static UE_API bool GetNamedInclusiveDiscreteObservation(TArray<FName>& OutNames, const ULearningAgentsObservationObject* Object, const FLearningAgentsObservationObjectElement Element, const FName Tag = TEXT("InclusiveDiscreteObservation"));
-	static UE_API bool GetNamedInclusiveDiscreteObservationToArrayView(TArrayView<FName> OutNames, const ULearningAgentsObservationObject* Object, const FLearningAgentsObservationObjectElement Element, const FName Tag = TEXT("InclusiveDiscreteObservation"));
+	static UE_API bool GetNamedInclusiveDiscreteObservation(TArray<FName>& OutNames, const ULearningAgentsObservationObject* Object, const FLearningAgentsObservationObjectElement Element, const bool bActive = true, const FName Tag = TEXT("InclusiveDiscreteObservation"));
+	static UE_API bool GetNamedInclusiveDiscreteObservationToArrayView(TArrayView<FName> OutNames, const ULearningAgentsObservationObject* Object, const FLearningAgentsObservationObjectElement Element, const bool bActive = true, const FName Tag = TEXT("InclusiveDiscreteObservation"));
 
 	UFUNCTION(BlueprintPure = false, Category = "LearningAgents", meta = (AdvancedDisplay = 3, ReturnDisplayName = "Success"))
-	static UE_API bool GetCountObservation(int32& OutNum, const ULearningAgentsObservationObject* Object, const FLearningAgentsObservationObjectElement Element, const int32 MaxNum, const FName Tag = TEXT("CountObservation"));
+	static UE_API bool GetCountObservation(int32& OutNum, const ULearningAgentsObservationObject* Object, const FLearningAgentsObservationObjectElement Element, const int32 MaxNum, const bool bActive = true, const FName Tag = TEXT("CountObservation"));
 
 	UFUNCTION(BlueprintPure = false, Category = "LearningAgents", meta = (AdvancedDisplay = 3, ReturnDisplayName = "Success"))
-	static UE_API bool GetStructObservationNum(int32& OutNum, const ULearningAgentsObservationObject* Object, const FLearningAgentsObservationObjectElement Element, const FName Tag = TEXT("StructObservation"));
+	static UE_API bool GetStructObservationNum(int32& OutNum, const ULearningAgentsObservationObject* Object, const FLearningAgentsObservationObjectElement Element, const bool bActive = true, const FName Tag = TEXT("StructObservation"));
 
 	UFUNCTION(BlueprintPure = false, Category = "LearningAgents", meta = (AdvancedDisplay = 3, ReturnDisplayName = "Success"))
-	static UE_API bool GetStructObservation(TMap<FName, FLearningAgentsObservationObjectElement>& OutElements, const ULearningAgentsObservationObject* Object, const FLearningAgentsObservationObjectElement Element, const FName Tag = TEXT("StructObservation"));
+	static UE_API bool GetStructObservation(TMap<FName, FLearningAgentsObservationObjectElement>& OutElements, const ULearningAgentsObservationObject* Object, const FLearningAgentsObservationObjectElement Element, const bool bActive = true, const FName Tag = TEXT("StructObservation"));
 	
 	UFUNCTION(BlueprintPure = false, Category = "LearningAgents", meta = (AdvancedDisplay = 4, ReturnDisplayName = "Success"))
-	static UE_API bool GetStructObservationElement(FLearningAgentsObservationObjectElement& OutElement, const ULearningAgentsObservationObject* Object, const FLearningAgentsObservationObjectElement Element, const FName ElementName, const FName Tag = TEXT("StructObservation"));
+	static UE_API bool GetStructObservationElement(FLearningAgentsObservationObjectElement& OutElement, const ULearningAgentsObservationObject* Object, const FLearningAgentsObservationObjectElement Element, const FName ElementName, const bool bActive = true, const FName Tag = TEXT("StructObservation"));
 
 	UFUNCTION(BlueprintPure = false, Category = "LearningAgents", meta = (AdvancedDisplay = 4, ReturnDisplayName = "Success"))
-	static UE_API bool GetStructObservationToArrays(TArray<FName>& OutElementNames, TArray<FLearningAgentsObservationObjectElement>& OutElements, const ULearningAgentsObservationObject* Object, const FLearningAgentsObservationObjectElement Element, const FName Tag = TEXT("StructObservation"));
-	static UE_API bool GetStructObservationToArrayViews(TArrayView<FName> OutElementNames, TArrayView<FLearningAgentsObservationObjectElement> OutElements, const ULearningAgentsObservationObject* Object, const FLearningAgentsObservationObjectElement Element, const FName Tag = TEXT("StructObservation"));
+	static UE_API bool GetStructObservationToArrays(TArray<FName>& OutElementNames, TArray<FLearningAgentsObservationObjectElement>& OutElements, const ULearningAgentsObservationObject* Object, const FLearningAgentsObservationObjectElement Element, const bool bActive = true, const FName Tag = TEXT("StructObservation"));
+	static UE_API bool GetStructObservationToArrayViews(TArrayView<FName> OutElementNames, TArrayView<FLearningAgentsObservationObjectElement> OutElements, const ULearningAgentsObservationObject* Object, const FLearningAgentsObservationObjectElement Element, const bool bActive = true, const FName Tag = TEXT("StructObservation"));
 
 	UFUNCTION(BlueprintPure = false, Category = "LearningAgents", meta = (AdvancedDisplay = 4, ReturnDisplayName = "Success"))
-	static UE_API bool GetExclusiveUnionObservation(FName& OutElementName, FLearningAgentsObservationObjectElement& OutElement, const ULearningAgentsObservationObject* Object, const FLearningAgentsObservationObjectElement Element, const FName Tag = TEXT("ExclusiveUnionObservation"));
+	static UE_API bool GetExclusiveUnionObservation(FName& OutElementName, FLearningAgentsObservationObjectElement& OutElement, const ULearningAgentsObservationObject* Object, const FLearningAgentsObservationObjectElement Element, const bool bActive = true, const FName Tag = TEXT("ExclusiveUnionObservation"));
 
 	UFUNCTION(BlueprintPure = false, Category = "LearningAgents", meta = (AdvancedDisplay = 3, ReturnDisplayName = "Success"))
-	static UE_API bool GetInclusiveUnionObservationNum(int32& OutNum, const ULearningAgentsObservationObject* Object, const FLearningAgentsObservationObjectElement Element, const FName Tag = TEXT("InclusiveUnionObservation"));
+	static UE_API bool GetInclusiveUnionObservationNum(int32& OutNum, const ULearningAgentsObservationObject* Object, const FLearningAgentsObservationObjectElement Element, const bool bActive = true, const FName Tag = TEXT("InclusiveUnionObservation"));
 
 	UFUNCTION(BlueprintPure = false, Category = "LearningAgents", meta = (AdvancedDisplay = 3, ReturnDisplayName = "Success"))
-	static UE_API bool GetInclusiveUnionObservation(TMap<FName, FLearningAgentsObservationObjectElement>& OutElements, const ULearningAgentsObservationObject* Object, const FLearningAgentsObservationObjectElement Element, const FName Tag = TEXT("InclusiveUnionObservation"));
+	static UE_API bool GetInclusiveUnionObservation(TMap<FName, FLearningAgentsObservationObjectElement>& OutElements, const ULearningAgentsObservationObject* Object, const FLearningAgentsObservationObjectElement Element, const bool bActive = true, const FName Tag = TEXT("InclusiveUnionObservation"));
 
 	UFUNCTION(BlueprintPure = false, Category = "LearningAgents", meta = (AdvancedDisplay = 4, ReturnDisplayName = "Success"))
-	static UE_API bool GetInclusiveUnionObservationToArrays(TArray<FName>& OutElementNames, TArray<FLearningAgentsObservationObjectElement>& OutElements, const ULearningAgentsObservationObject* Object, const FLearningAgentsObservationObjectElement Element, const FName Tag = TEXT("InclusiveUnionObservation"));
-	static UE_API bool GetInclusiveUnionObservationToArrayViews(TArrayView<FName> OutElementNames, TArrayView<FLearningAgentsObservationObjectElement> OutElements, const ULearningAgentsObservationObject* Object, const FLearningAgentsObservationObjectElement Element, const FName Tag = TEXT("InclusiveUnionObservation"));
+	static UE_API bool GetInclusiveUnionObservationToArrays(TArray<FName>& OutElementNames, TArray<FLearningAgentsObservationObjectElement>& OutElements, const ULearningAgentsObservationObject* Object, const FLearningAgentsObservationObjectElement Element, const bool bActive = true, const FName Tag = TEXT("InclusiveUnionObservation"));
+	static UE_API bool GetInclusiveUnionObservationToArrayViews(TArrayView<FName> OutElementNames, TArrayView<FLearningAgentsObservationObjectElement> OutElements, const ULearningAgentsObservationObject* Object, const FLearningAgentsObservationObjectElement Element, const bool bActive = true, const FName Tag = TEXT("InclusiveUnionObservation"));
 
 	UFUNCTION(BlueprintPure = false, Category = "LearningAgents", meta = (AdvancedDisplay = 3, ReturnDisplayName = "Success"))
-	static UE_API bool GetStaticArrayObservationNum(int32& OutNum, const ULearningAgentsObservationObject* Object, const FLearningAgentsObservationObjectElement Element, const FName Tag = TEXT("StaticArrayObservation"));
+	static UE_API bool GetStaticArrayObservationNum(int32& OutNum, const ULearningAgentsObservationObject* Object, const FLearningAgentsObservationObjectElement Element, const bool bActive = true, const FName Tag = TEXT("StaticArrayObservation"));
 
 	UFUNCTION(BlueprintPure = false, Category = "LearningAgents", meta = (AdvancedDisplay = 3, ReturnDisplayName = "Success"))
-	static UE_API bool GetStaticArrayObservation(TArray<FLearningAgentsObservationObjectElement>& OutElements, const ULearningAgentsObservationObject* Object, const FLearningAgentsObservationObjectElement Element, const FName Tag = TEXT("StaticArrayObservation"));
-	static UE_API bool GetStaticArrayObservationToArrayView(TArrayView<FLearningAgentsObservationObjectElement> OutElements, const ULearningAgentsObservationObject* Object, const FLearningAgentsObservationObjectElement Element, const FName Tag = TEXT("StaticArrayObservation"));
+	static UE_API bool GetStaticArrayObservation(TArray<FLearningAgentsObservationObjectElement>& OutElements, const ULearningAgentsObservationObject* Object, const FLearningAgentsObservationObjectElement Element, const bool bActive = true, const FName Tag = TEXT("StaticArrayObservation"));
+	static UE_API bool GetStaticArrayObservationToArrayView(TArrayView<FLearningAgentsObservationObjectElement> OutElements, const ULearningAgentsObservationObject* Object, const FLearningAgentsObservationObjectElement Element, const bool bActive = true, const FName Tag = TEXT("StaticArrayObservation"));
 
 	UFUNCTION(BlueprintPure = false, Category = "LearningAgents", meta = (AdvancedDisplay = 3, ReturnDisplayName = "Success"))
-	static UE_API bool GetSetObservationNum(int32& OutNum, const ULearningAgentsObservationObject* Object, const FLearningAgentsObservationObjectElement Element, const FName Tag = TEXT("SetObservation"));
+	static UE_API bool GetSetObservationNum(int32& OutNum, const ULearningAgentsObservationObject* Object, const FLearningAgentsObservationObjectElement Element, const bool bActive = true, const FName Tag = TEXT("SetObservation"));
 
 	UFUNCTION(BlueprintPure = false, Category = "LearningAgents", meta = (AdvancedDisplay = 3, ReturnDisplayName = "Success"))
-	static UE_API bool GetSetObservation(TSet<FLearningAgentsObservationObjectElement>& OutElements, const ULearningAgentsObservationObject* Object, const FLearningAgentsObservationObjectElement Element, const FName Tag = TEXT("SetObservation"));
+	static UE_API bool GetSetObservation(TSet<FLearningAgentsObservationObjectElement>& OutElements, const ULearningAgentsObservationObject* Object, const FLearningAgentsObservationObjectElement Element, const bool bActive = true, const FName Tag = TEXT("SetObservation"));
 
 	UFUNCTION(BlueprintPure = false, Category = "LearningAgents", meta = (AdvancedDisplay = 3, ReturnDisplayName = "Success"))
-	static UE_API bool GetSetObservationToArray(TArray<FLearningAgentsObservationObjectElement>& OutElements, const ULearningAgentsObservationObject* Object, const FLearningAgentsObservationObjectElement Element, const FName Tag = TEXT("SetObservation"));
-	static UE_API bool GetSetObservationToArrayView(TArrayView<FLearningAgentsObservationObjectElement> OutElements, const ULearningAgentsObservationObject* Object, const FLearningAgentsObservationObjectElement Element, const FName Tag = TEXT("SetObservation"));
+	static UE_API bool GetSetObservationToArray(TArray<FLearningAgentsObservationObjectElement>& OutElements, const ULearningAgentsObservationObject* Object, const FLearningAgentsObservationObjectElement Element, const bool bActive = true, const FName Tag = TEXT("SetObservation"));
+	static UE_API bool GetSetObservationToArrayView(TArrayView<FLearningAgentsObservationObjectElement> OutElements, const ULearningAgentsObservationObject* Object, const FLearningAgentsObservationObjectElement Element, const bool bActive = true, const FName Tag = TEXT("SetObservation"));
 
 	UFUNCTION(BlueprintPure = false, Category = "LearningAgents", meta = (AdvancedDisplay = 4, ReturnDisplayName = "Success"))
-	static UE_API bool GetPairObservation(FLearningAgentsObservationObjectElement& OutKey, FLearningAgentsObservationObjectElement& OutValue, const ULearningAgentsObservationObject* Object, const FLearningAgentsObservationObjectElement Element, const FName Tag = TEXT("PairObservation"));
+	static UE_API bool GetPairObservation(FLearningAgentsObservationObjectElement& OutKey, FLearningAgentsObservationObjectElement& OutValue, const ULearningAgentsObservationObject* Object, const FLearningAgentsObservationObjectElement Element, const bool bActive = true, const FName Tag = TEXT("PairObservation"));
 
 	UFUNCTION(BlueprintPure = false, Category = "LearningAgents", meta = (AdvancedDisplay = 3, ReturnDisplayName = "Success"))
-	static UE_API bool GetArrayObservationNum(int32& OutNum, const ULearningAgentsObservationObject* Object, const FLearningAgentsObservationObjectElement Element, const FName Tag = TEXT("ArrayObservation"));
+	static UE_API bool GetArrayObservationNum(int32& OutNum, const ULearningAgentsObservationObject* Object, const FLearningAgentsObservationObjectElement Element, const bool bActive = true, const FName Tag = TEXT("ArrayObservation"));
 
 	UFUNCTION(BlueprintPure = false, Category = "LearningAgents", meta = (AdvancedDisplay = 3, ReturnDisplayName = "Success"))
-	static UE_API bool GetArrayObservation(TArray<FLearningAgentsObservationObjectElement>& OutElements, const ULearningAgentsObservationObject* Object, const FLearningAgentsObservationObjectElement Element, const int32 MaxNum, const FName Tag = TEXT("ArrayObservation"));
-	static UE_API bool GetArrayObservationToArrayView(TArrayView<FLearningAgentsObservationObjectElement> OutElements, const ULearningAgentsObservationObject* Object, const FLearningAgentsObservationObjectElement Element, const int32 MaxNum, const FName Tag = TEXT("ArrayObservation"));
+	static UE_API bool GetArrayObservation(TArray<FLearningAgentsObservationObjectElement>& OutElements, const ULearningAgentsObservationObject* Object, const FLearningAgentsObservationObjectElement Element, const int32 MaxNum, const bool bActive = true, const FName Tag = TEXT("ArrayObservation"));
+	static UE_API bool GetArrayObservationToArrayView(TArrayView<FLearningAgentsObservationObjectElement> OutElements, const ULearningAgentsObservationObject* Object, const FLearningAgentsObservationObjectElement Element, const int32 MaxNum, const bool bActive = true, const FName Tag = TEXT("ArrayObservation"));
 
 	UFUNCTION(BlueprintPure = false, Category = "LearningAgents", meta = (AdvancedDisplay = 3, ReturnDisplayName = "Success"))
-	static UE_API bool GetMapObservationNum(int32& OutNum, const ULearningAgentsObservationObject* Object, const FLearningAgentsObservationObjectElement Element, const FName Tag = TEXT("MapObservation"));
+	static UE_API bool GetMapObservationNum(int32& OutNum, const ULearningAgentsObservationObject* Object, const FLearningAgentsObservationObjectElement Element, const bool bActive = true, const FName Tag = TEXT("MapObservation"));
 
 	UFUNCTION(BlueprintPure = false, Category = "LearningAgents", meta = (AdvancedDisplay = 3, ReturnDisplayName = "Success"))
-	static UE_API bool GetMapObservation(TMap<FLearningAgentsObservationObjectElement, FLearningAgentsObservationObjectElement>& OutElements, const ULearningAgentsObservationObject* Object, const FLearningAgentsObservationObjectElement Element, const FName Tag = TEXT("MapObservation"));
+	static UE_API bool GetMapObservation(TMap<FLearningAgentsObservationObjectElement, FLearningAgentsObservationObjectElement>& OutElements, const ULearningAgentsObservationObject* Object, const FLearningAgentsObservationObjectElement Element, const bool bActive = true, const FName Tag = TEXT("MapObservation"));
 
 	UFUNCTION(BlueprintPure = false, Category = "LearningAgents", meta = (AdvancedDisplay = 4, ReturnDisplayName = "Success"))
-	static UE_API bool GetMapObservationToArrays(TArray<FLearningAgentsObservationObjectElement>& OutKeys, TArray<FLearningAgentsObservationObjectElement>& OutValues, const ULearningAgentsObservationObject* Object, const FLearningAgentsObservationObjectElement Element, const FName Tag = TEXT("MapObservation"));
-	static UE_API bool GetMapObservationToArrayViews(TArrayView<FLearningAgentsObservationObjectElement> OutKeys, TArrayView<FLearningAgentsObservationObjectElement> OutValues, const ULearningAgentsObservationObject* Object, const FLearningAgentsObservationObjectElement Element, const FName Tag = TEXT("MapObservation"));
+	static UE_API bool GetMapObservationToArrays(TArray<FLearningAgentsObservationObjectElement>& OutKeys, TArray<FLearningAgentsObservationObjectElement>& OutValues, const ULearningAgentsObservationObject* Object, const FLearningAgentsObservationObjectElement Element, const bool bActive = true, const FName Tag = TEXT("MapObservation"));
+	static UE_API bool GetMapObservationToArrayViews(TArrayView<FLearningAgentsObservationObjectElement> OutKeys, TArrayView<FLearningAgentsObservationObjectElement> OutValues, const ULearningAgentsObservationObject* Object, const FLearningAgentsObservationObjectElement Element, const bool bActive = true, const FName Tag = TEXT("MapObservation"));
 
 	UFUNCTION(BlueprintPure = false, Category = "LearningAgents", meta = (AdvancedDisplay = 4, ReturnDisplayName = "Success"))
-	static UE_API bool GetEnumObservation(uint8& OutEnumValue, const ULearningAgentsObservationObject* Object, const FLearningAgentsObservationObjectElement Element, const UEnum* Enum, const FName Tag = TEXT("EnumObservation"));
+	static UE_API bool GetEnumObservation(uint8& OutEnumValue, const ULearningAgentsObservationObject* Object, const FLearningAgentsObservationObjectElement Element, const UEnum* Enum, const bool bActive = true, const FName Tag = TEXT("EnumObservation"));
 
 	UFUNCTION(BlueprintPure = false, Category = "LearningAgents", meta = (AdvancedDisplay = 4, ReturnDisplayName = "Success"))
-	static UE_API bool GetBitmaskObservation(int32& OutBitmaskValue, const ULearningAgentsObservationObject* Object, const FLearningAgentsObservationObjectElement Element, const UEnum* Enum, const FName Tag = TEXT("BitmaskObservation"));
+	static UE_API bool GetBitmaskObservation(int32& OutBitmaskValue, const ULearningAgentsObservationObject* Object, const FLearningAgentsObservationObjectElement Element, const UEnum* Enum, const bool bActive = true, const FName Tag = TEXT("BitmaskObservation"));
 
 	UFUNCTION(BlueprintPure = false, Category = "LearningAgents", meta = (AdvancedDisplay = 3, ExpandEnumAsExecs = "OutOption", ReturnDisplayName = "Success"))
-	static UE_API bool GetOptionalObservation(ELearningAgentsOptionalObservation& OutOption, FLearningAgentsObservationObjectElement& OutElement, const ULearningAgentsObservationObject* Object, const FLearningAgentsObservationObjectElement Element, const FName Tag = TEXT("OptionalObservation"));
+	static UE_API bool GetOptionalObservation(ELearningAgentsOptionalObservation& OutOption, FLearningAgentsObservationObjectElement& OutElement, const ULearningAgentsObservationObject* Object, const FLearningAgentsObservationObjectElement Element, const bool bActive = true, const FName Tag = TEXT("OptionalObservation"));
 
 	UFUNCTION(BlueprintPure = false, Category = "LearningAgents", meta = (AdvancedDisplay = 3, ExpandEnumAsExecs = "OutEither", ReturnDisplayName = "Success"))
-	static UE_API bool GetEitherObservation(ELearningAgentsEitherObservation& OutEither, FLearningAgentsObservationObjectElement& OutElement, const ULearningAgentsObservationObject* Object, const FLearningAgentsObservationObjectElement Element, const FName Tag = TEXT("EitherObservation"));
+	static UE_API bool GetEitherObservation(ELearningAgentsEitherObservation& OutEither, FLearningAgentsObservationObjectElement& OutElement, const ULearningAgentsObservationObject* Object, const FLearningAgentsObservationObjectElement Element, const bool bActive = true, const FName Tag = TEXT("EitherObservation"));
 
 	UFUNCTION(BlueprintPure = false, Category = "LearningAgents", meta = (AdvancedDisplay = 3, ReturnDisplayName = "Success"))
-	static UE_API bool GetEncodingObservation(FLearningAgentsObservationObjectElement& OutElement, const ULearningAgentsObservationObject* Object, const FLearningAgentsObservationObjectElement Element, const FName Tag = TEXT("EncodingObservation"));
+	static UE_API bool GetEncodingObservation(FLearningAgentsObservationObjectElement& OutElement, const ULearningAgentsObservationObject* Object, const FLearningAgentsObservationObjectElement Element, const bool bActive = true, const FName Tag = TEXT("EncodingObservation"));
 
 	UFUNCTION(BlueprintPure = false, Category = "LearningAgents", meta = (AdvancedDisplay = 3, ReturnDisplayName = "Success"))
-	static UE_API bool GetBoolObservation(bool& bOutValue, const ULearningAgentsObservationObject* Object, const FLearningAgentsObservationObjectElement Element, const FName Tag = TEXT("BoolObservation"));
+	static UE_API bool GetBoolObservation(bool& bOutValue, const ULearningAgentsObservationObject* Object, const FLearningAgentsObservationObjectElement Element, const bool bActive = true, const FName Tag = TEXT("BoolObservation"));
 
 	UFUNCTION(BlueprintPure = false, Category = "LearningAgents", meta = (AdvancedDisplay = 3, ReturnDisplayName = "Success"))
-	static UE_API bool GetFloatObservation(float& OutValue, const ULearningAgentsObservationObject* Object, const FLearningAgentsObservationObjectElement Element, const FName Tag = TEXT("FloatObservation"));
+	static UE_API bool GetFloatObservation(float& OutValue, const ULearningAgentsObservationObject* Object, const FLearningAgentsObservationObjectElement Element, const bool bActive = true, const FName Tag = TEXT("FloatObservation"));
 
 	UFUNCTION(BlueprintPure = false, Category = "LearningAgents", meta = (AdvancedDisplay = 4, ReturnDisplayName = "Success"))
-	static UE_API bool GetLocationObservation(FVector& OutLocation, const ULearningAgentsObservationObject* Object, const FLearningAgentsObservationObjectElement Element, const FTransform RelativeTransform = FTransform(), const FName Tag = TEXT("LocationObservation"));
+	static UE_API bool GetLocationObservation(FVector& OutLocation, const ULearningAgentsObservationObject* Object, const FLearningAgentsObservationObjectElement Element, const FTransform RelativeTransform = FTransform(), const bool bActive = true, const FName Tag = TEXT("LocationObservation"));
 
 	UFUNCTION(BlueprintPure = false, Category = "LearningAgents", meta = (AdvancedDisplay = 4, ReturnDisplayName = "Success"))
-	static UE_API bool GetRotationObservation(FRotator& OutRotation, const ULearningAgentsObservationObject* Object, const FLearningAgentsObservationObjectElement Element, const FRotator RelativeRotation = FRotator::ZeroRotator, const FName Tag = TEXT("RotationObservation"));
+	static UE_API bool GetRotationObservation(FRotator& OutRotation, const ULearningAgentsObservationObject* Object, const FLearningAgentsObservationObjectElement Element, const FRotator RelativeRotation = FRotator::ZeroRotator, const bool bActive = true, const FName Tag = TEXT("RotationObservation"));
 
 	UFUNCTION(BlueprintPure = false, Category = "LearningAgents", meta = (AdvancedDisplay = 4, ReturnDisplayName = "Success"))
-	static UE_API bool GetRotationObservationAsQuat(FQuat& OutRotation, const ULearningAgentsObservationObject* Object, const FLearningAgentsObservationObjectElement Element, const FQuat RelativeRotation, const FName Tag = TEXT("RotationObservation"));
+	static UE_API bool GetRotationObservationAsQuat(FQuat& OutRotation, const ULearningAgentsObservationObject* Object, const FLearningAgentsObservationObjectElement Element, const FQuat RelativeRotation, const bool bActive = true, const FName Tag = TEXT("RotationObservation"));
 
 	UFUNCTION(BlueprintPure = false, Category = "LearningAgents", meta = (AdvancedDisplay = 4, ReturnDisplayName = "Success"))
-	static UE_API bool GetScaleObservation(FVector& OutScale, const ULearningAgentsObservationObject* Object, const FLearningAgentsObservationObjectElement Element, const FVector RelativeScale = FVector(1, 1, 1), const FName Tag = TEXT("ScaleObservation"));
+	static UE_API bool GetScaleObservation(FVector& OutScale, const ULearningAgentsObservationObject* Object, const FLearningAgentsObservationObjectElement Element, const FVector RelativeScale = FVector(1, 1, 1), const bool bActive = true, const FName Tag = TEXT("ScaleObservation"));
 
 	UFUNCTION(BlueprintPure = false, Category = "LearningAgents", meta = (AdvancedDisplay = 4, ReturnDisplayName = "Success"))
-	static UE_API bool GetTransformObservation(FTransform& OutTransform, const ULearningAgentsObservationObject* Object, const FLearningAgentsObservationObjectElement Element, const FTransform RelativeTransform = FTransform(), const FName Tag = TEXT("TransformObservation"));
+	static UE_API bool GetTransformObservation(FTransform& OutTransform, const ULearningAgentsObservationObject* Object, const FLearningAgentsObservationObjectElement Element, const FTransform RelativeTransform = FTransform(), const bool bActive = true, const FName Tag = TEXT("TransformObservation"));
 
 	UFUNCTION(BlueprintPure = false, Category = "LearningAgents", meta = (AdvancedDisplay = 4, ReturnDisplayName = "Success"))
-	static UE_API bool GetAngleObservation(float& OutAngle, const ULearningAgentsObservationObject* Object, const FLearningAgentsObservationObjectElement Element, const float RelativeAngle = 0.0f, const FName Tag = TEXT("AngleObservation"));
+	static UE_API bool GetAngleObservation(float& OutAngle, const ULearningAgentsObservationObject* Object, const FLearningAgentsObservationObjectElement Element, const float RelativeAngle = 0.0f, const bool bActive = true, const FName Tag = TEXT("AngleObservation"));
 
 	UFUNCTION(BlueprintPure = false, Category = "LearningAgents", meta = (AdvancedDisplay = 4, ReturnDisplayName = "Success"))
-	static UE_API bool GetAngleObservationRadians(float& OutAngle, const ULearningAgentsObservationObject* Object, const FLearningAgentsObservationObjectElement Element, const float RelativeAngle = 0.0f, const FName Tag = TEXT("AngleObservation"));
+	static UE_API bool GetAngleObservationRadians(float& OutAngle, const ULearningAgentsObservationObject* Object, const FLearningAgentsObservationObjectElement Element, const float RelativeAngle = 0.0f, const bool bActive = true, const FName Tag = TEXT("AngleObservation"));
 
 	UFUNCTION(BlueprintPure = false, Category = "LearningAgents", meta = (AdvancedDisplay = 4, ReturnDisplayName = "Success"))
-	static UE_API bool GetVelocityObservation(FVector& OutVelocity, const ULearningAgentsObservationObject* Object, const FLearningAgentsObservationObjectElement Element, const FTransform RelativeTransform = FTransform(), const FName Tag = TEXT("VelocityObservation"));
+	static UE_API bool GetVelocityObservation(FVector& OutVelocity, const ULearningAgentsObservationObject* Object, const FLearningAgentsObservationObjectElement Element, const FTransform RelativeTransform = FTransform(), const bool bActive = true, const FName Tag = TEXT("VelocityObservation"));
 
 	UFUNCTION(BlueprintPure = false, Category = "LearningAgents", meta = (AdvancedDisplay = 4, ReturnDisplayName = "Success"))
-	static UE_API bool GetDirectionObservation(FVector& OutDirection, const ULearningAgentsObservationObject* Object, const FLearningAgentsObservationObjectElement Element, const FTransform RelativeTransform = FTransform(), const FName Tag = TEXT("DirectionObservation"));
+	static UE_API bool GetDirectionObservation(FVector& OutDirection, const ULearningAgentsObservationObject* Object, const FLearningAgentsObservationObjectElement Element, const FTransform RelativeTransform = FTransform(), const bool bActive = true, const FName Tag = TEXT("DirectionObservation"));
 
 	// Get Spline Observations
 
 	UFUNCTION(BlueprintPure = false, Category = "LearningAgents", meta = (AdvancedDisplay = 4, ReturnDisplayName = "Success"))
-	static UE_API bool GetLocationAlongSplineObservation(FVector& OutLocation, const ULearningAgentsObservationObject* Object, const FLearningAgentsObservationObjectElement Element, const FTransform RelativeTransform = FTransform(), const FName Tag = TEXT("LocationAlongSplineObservation"));
+	static UE_API bool GetLocationAlongSplineObservation(FVector& OutLocation, const ULearningAgentsObservationObject* Object, const FLearningAgentsObservationObjectElement Element, const FTransform RelativeTransform = FTransform(), const bool bActive = true, const FName Tag = TEXT("LocationAlongSplineObservation"));
 
 	UFUNCTION(BlueprintPure = false, Category = "LearningAgents", meta = (AdvancedDisplay = 5, ReturnDisplayName = "Success"))
-	static UE_API bool GetProportionAlongSplineObservation(bool& bOutIsClosedLoop, float& OutAngle, float& OutPropotion, const ULearningAgentsObservationObject* Object, const FLearningAgentsObservationObjectElement Element, const FName Tag = TEXT("ProportionAlongSplineObservation"));
+	static UE_API bool GetProportionAlongSplineObservation(bool& bOutIsClosedLoop, float& OutAngle, float& OutPropotion, const ULearningAgentsObservationObject* Object, const FLearningAgentsObservationObjectElement Element, const bool bActive = true, const FName Tag = TEXT("ProportionAlongSplineObservation"));
 
 	UFUNCTION(BlueprintPure = false, Category = "LearningAgents", meta = (AdvancedDisplay = 4, ReturnDisplayName = "Success"))
-	static UE_API bool GetDirectionAlongSplineObservation(FVector& OutDirection, const ULearningAgentsObservationObject* Object, const FLearningAgentsObservationObjectElement Element, const FTransform RelativeTransform = FTransform(), const FName Tag = TEXT("DirectionAlongSplineObservation"));
+	static UE_API bool GetDirectionAlongSplineObservation(FVector& OutDirection, const ULearningAgentsObservationObject* Object, const FLearningAgentsObservationObjectElement Element, const FTransform RelativeTransform = FTransform(), const bool bActive = true, const FName Tag = TEXT("DirectionAlongSplineObservation"));
 
 	// Get Ray Cast Observations
 
 	UFUNCTION(BlueprintPure, Category = "LearningAgents", meta = (AdvancedDisplay = 3, ReturnDisplayName = "Success"))
-	static UE_API bool GetProportionAlongRayObservation(float& OutProportion, const ULearningAgentsObservationObject* Object, const FLearningAgentsObservationObjectElement Element, const FName Tag = TEXT("ProportionAlongRayObservation"));
+	static UE_API bool GetProportionAlongRayObservation(float& OutProportion, const ULearningAgentsObservationObject* Object, const FLearningAgentsObservationObjectElement Element, const bool bActive = true, const FName Tag = TEXT("ProportionAlongRayObservation"));
 
 	UFUNCTION(BlueprintPure = false, Category = "LearningAgents", meta = (AdvancedDisplay = 3, ReturnDisplayName = "Success"))
-	static UE_API bool GetConv1dObservation(FLearningAgentsObservationObjectElement& OutElement, const ULearningAgentsObservationObject* Object, const FLearningAgentsObservationObjectElement Element, const FName Tag = TEXT("Conv1dObservation"));
+	static UE_API bool GetConv1dObservation(FLearningAgentsObservationObjectElement& OutElement, const ULearningAgentsObservationObject* Object, const FLearningAgentsObservationObjectElement Element, const bool bActive = true, const FName Tag = TEXT("Conv1dObservation"));
 	
 	UFUNCTION(BlueprintPure = false, Category = "LearningAgents", meta = (AdvancedDisplay = 3, ReturnDisplayName = "Success"))
-	static UE_API bool GetConv2dObservation(FLearningAgentsObservationObjectElement& OutElement, const ULearningAgentsObservationObject* Object, const FLearningAgentsObservationObjectElement Element, const FName Tag = TEXT("Conv2dObservation"));
+	static UE_API bool GetConv2dObservation(FLearningAgentsObservationObjectElement& OutElement, const ULearningAgentsObservationObject* Object, const FLearningAgentsObservationObjectElement Element, const bool bActive = true, const FName Tag = TEXT("Conv2dObservation"));
 };
 
 
