[Backout] - CL50499274
#rnx
Original CL Desc
-----------------------------------------------------------------
READY TO SUBMIT

[LearningAgents]
- added bActive flag to all observations and actions that route to NullObservation / NullAction if inactive

[LACombat]
- updated combat interactor toward data oriented design: all actions/data configs can be loaded from one source, no more specialized structs for scales/biases etc
- header reduced from 1300 to <200 lines
- split out all rule based guard logic like lethality reducer/action stickiness/guards into a separate component so the interactor is pure ML specs
- kept old version of interactor for legacy and safety

#Synced-CL 50502114

diff --git a/Engine/Plugins/Experimental/LearningAgents/Source/LearningAgents/Private/LearningAgentsActions.cpp b/Engine/Plugins/Experimental/LearningAgents/Source/LearningAgents/Private/LearningAgentsActions.cpp
--- a/Engine/Plugins/Experimental/LearningAgents/Source/LearningAgents/Private/LearningAgentsActions.cpp
+++ b/Engine/Plugins/Experimental/LearningAgents/Source/LearningAgents/Private/LearningAgentsActions.cpp
@@ -1060,19 +1060,14 @@ FLearningAgentsActionSchemaElement ULearningAgentsActions::SpecifyNullAction(ULe
 	return { Schema->ActionSchema.CreateNull(Tag) };
 }
 
-FLearningAgentsActionSchemaElement ULearningAgentsActions::SpecifyContinuousAction(ULearningAgentsActionSchema* Schema, const int32 Size, const float Scale, const bool bActive, const FName Tag)
+FLearningAgentsActionSchemaElement ULearningAgentsActions::SpecifyContinuousAction(ULearningAgentsActionSchema* Schema, const int32 Size, const float Scale, const FName Tag)
 {
-	if (!bActive)
-	{
-		return SpecifyNullAction(Schema, Tag);
-	}
-
 	if (!Schema)
 	{
 		UE_LOG(LogLearning, Error, TEXT("SpecifyContinuousAction: Schema is nullptr."));
 		return FLearningAgentsActionSchemaElement();
 	}
-	
+
 	if (Size < 0)
 	{
 		UE_LOG(LogLearning, Error, TEXT("SpecifyContinuousAction: Invalid Continuous Action Size '%i'."), Size);
@@ -1093,23 +1088,13 @@ FLearningAgentsActionSchemaElement ULearningAgentsActions::SpecifyContinuousActi
 	return { Schema->ActionSchema.CreateContinuous({ Size, Scale }, Tag) };
 }
 
-FLearningAgentsActionSchemaElement ULearningAgentsActions::SpecifyExclusiveDiscreteAction(ULearningAgentsActionSchema* Schema, const int32 Size, const TArray<float>& PriorProbabilities, const bool bActive, const FName Tag)
+FLearningAgentsActionSchemaElement ULearningAgentsActions::SpecifyExclusiveDiscreteAction(ULearningAgentsActionSchema* Schema, const int32 Size, const TArray<float>& PriorProbabilities, const FName Tag)
 {
-	if (!bActive)
-	{
-		return SpecifyNullAction(Schema, Tag);
-	}
-
-	return SpecifyExclusiveDiscreteActionFromArrayView(Schema, Size, PriorProbabilities, bActive, Tag);
+	return SpecifyExclusiveDiscreteActionFromArrayView(Schema, Size, PriorProbabilities, Tag);
 }
 
-FLearningAgentsActionSchemaElement ULearningAgentsActions::SpecifyExclusiveDiscreteActionFromArrayView(ULearningAgentsActionSchema* Schema, const int32 Size, const TArrayView<const float> PriorProbabilities, const bool bActive, const FName Tag)
+FLearningAgentsActionSchemaElement ULearningAgentsActions::SpecifyExclusiveDiscreteActionFromArrayView(ULearningAgentsActionSchema* Schema, const int32 Size, const TArrayView<const float> PriorProbabilities, const FName Tag)
 {
-	if (!bActive)
-	{
-		return SpecifyNullAction(Schema, Tag);
-	}
-
 	if (!Schema)
 	{
 		UE_LOG(LogLearning, Error, TEXT("SpecifyExclusiveDiscreteActionFromArrayView: Schema is nullptr."));
@@ -1144,13 +1129,8 @@ FLearningAgentsActionSchemaElement ULearningAgentsActions::SpecifyExclusiveDiscr
 	return { Schema->ActionSchema.CreateDiscreteExclusive({ Size, MakeArrayView(NormalizedPriorProbabilities) }, Tag)};
 }
 
-FLearningAgentsActionSchemaElement ULearningAgentsActions::SpecifyNamedExclusiveDiscreteAction(ULearningAgentsActionSchema* Schema, const TArray<FName>& Names, const TMap<FName, float>& PriorProbabilities, const bool bActive, const FName Tag)
+FLearningAgentsActionSchemaElement ULearningAgentsActions::SpecifyNamedExclusiveDiscreteAction(ULearningAgentsActionSchema* Schema, const TArray<FName>& Names, const TMap<FName, float>& PriorProbabilities, const FName Tag)
 {
-	if (!bActive)
-	{
-		return SpecifyNullAction(Schema, Tag);
-	}
-
 	const int32 NameNum = Names.Num();
 
 	if (NameNum == 0)
@@ -1189,16 +1169,11 @@ FLearningAgentsActionSchemaElement ULearningAgentsActions::SpecifyNamedExclusive
 		SortedPriorProbabilities[Idx] = NamePriorProbabilities[NameIndices[Idx]];
 	}
 
-	return SpecifyNamedExclusiveDiscreteActionFromArrayViews(Schema, SortedNames, SortedPriorProbabilities, bActive, Tag);
+	return SpecifyNamedExclusiveDiscreteActionFromArrayViews(Schema, SortedNames, SortedPriorProbabilities, Tag);
 }
 
-FLearningAgentsActionSchemaElement ULearningAgentsActions::SpecifyNamedExclusiveDiscreteActionFromArrayViews(ULearningAgentsActionSchema* Schema, const TArrayView<const FName> Names, const TArrayView<const float> PriorProbabilities, const bool bActive, const FName Tag)
+FLearningAgentsActionSchemaElement ULearningAgentsActions::SpecifyNamedExclusiveDiscreteActionFromArrayViews(ULearningAgentsActionSchema* Schema, const TArrayView<const FName> Names, const TArrayView<const float> PriorProbabilities, const FName Tag)
 {
-	if (!bActive)
-	{
-		return SpecifyNullAction(Schema, Tag);
-	}
-
 	if (!Schema)
 	{
 		UE_LOG(LogLearning, Error, TEXT("SpecifyNamedExclusiveDiscreteActionFromArrayViews: Schema is nullptr."));
@@ -1236,23 +1211,13 @@ FLearningAgentsActionSchemaElement ULearningAgentsActions::SpecifyNamedExclusive
 	return { Schema->ActionSchema.CreateNamedDiscreteExclusive({ Names, MakeArrayView(NormalizedPriorProbabilities) }, Tag) };
 }
 
-FLearningAgentsActionSchemaElement ULearningAgentsActions::SpecifyInclusiveDiscreteAction(ULearningAgentsActionSchema* Schema, const int32 Size, const TArray<float>& PriorProbabilities, const bool bActive, const FName Tag)
+FLearningAgentsActionSchemaElement ULearningAgentsActions::SpecifyInclusiveDiscreteAction(ULearningAgentsActionSchema* Schema, const int32 Size, const TArray<float>& PriorProbabilities, const FName Tag)
 {
-	if (!bActive)
-	{
-		return SpecifyNullAction(Schema, Tag);
-	}
-
-	return SpecifyInclusiveDiscreteActionFromArrayView(Schema, Size, PriorProbabilities, bActive, Tag);
+	return SpecifyInclusiveDiscreteActionFromArrayView(Schema, Size, PriorProbabilities, Tag);
 }
 
-FLearningAgentsActionSchemaElement ULearningAgentsActions::SpecifyInclusiveDiscreteActionFromArrayView(ULearningAgentsActionSchema* Schema, const int32 Size, const TArrayView<const float> PriorProbabilities, const bool bActive, const FName Tag)
+FLearningAgentsActionSchemaElement ULearningAgentsActions::SpecifyInclusiveDiscreteActionFromArrayView(ULearningAgentsActionSchema* Schema, const int32 Size, const TArrayView<const float> PriorProbabilities, const FName Tag)
 {
-	if (!bActive)
-	{
-		return SpecifyNullAction(Schema, Tag);
-	}
-
 	if (!Schema)
 	{
 		UE_LOG(LogLearning, Error, TEXT("SpecifyInclusiveDiscreteActionFromArrayView: Schema is nullptr."));
@@ -1287,13 +1252,8 @@ FLearningAgentsActionSchemaElement ULearningAgentsActions::SpecifyInclusiveDiscr
 	return { Schema->ActionSchema.CreateDiscreteInclusive({ Size, MakeArrayView(NormalizedPriorProbabilities) }, Tag) };
 }
 
-FLearningAgentsActionSchemaElement ULearningAgentsActions::SpecifyNamedInclusiveDiscreteAction(ULearningAgentsActionSchema* Schema, const TArray<FName> Names, const TMap<FName, float>& PriorProbabilities, const bool bActive, const FName Tag)
+FLearningAgentsActionSchemaElement ULearningAgentsActions::SpecifyNamedInclusiveDiscreteAction(ULearningAgentsActionSchema* Schema, const TArray<FName> Names, const TMap<FName, float>& PriorProbabilities, const FName Tag)
 {
-	if (!bActive)
-	{
-		return SpecifyNullAction(Schema, Tag);
-	}
-
 	const int32 NameNum = Names.Num();
 
 	if (NameNum == 0)
@@ -1332,16 +1292,11 @@ FLearningAgentsActionSchemaElement ULearningAgentsActions::SpecifyNamedInclusive
 		SortedPriorProbabilities[Idx] = NamePriorProbabilities[NameIndices[Idx]];
 	}
 
-	return SpecifyNamedInclusiveDiscreteActionFromArrayViews(Schema, SortedNames, SortedPriorProbabilities, bActive, Tag);
+	return SpecifyNamedInclusiveDiscreteActionFromArrayViews(Schema, SortedNames, SortedPriorProbabilities, Tag);
 }
 
-FLearningAgentsActionSchemaElement ULearningAgentsActions::SpecifyNamedInclusiveDiscreteActionFromArrayViews(ULearningAgentsActionSchema* Schema, const TArrayView<const FName> Names, const TArrayView<const float> PriorProbabilities, const bool bActive, const FName Tag)
+FLearningAgentsActionSchemaElement ULearningAgentsActions::SpecifyNamedInclusiveDiscreteActionFromArrayViews(ULearningAgentsActionSchema* Schema, const TArrayView<const FName> Names, const TArrayView<const float> PriorProbabilities, const FName Tag)
 {
-	if (!bActive)
-	{
-		return SpecifyNullAction(Schema, Tag);
-	}
-
 	if (!Schema)
 	{
 		UE_LOG(LogLearning, Error, TEXT("SpecifyNamedInclusiveDiscreteActionFromArrayViews: Schema is nullptr."));
@@ -1378,13 +1333,8 @@ FLearningAgentsActionSchemaElement ULearningAgentsActions::SpecifyNamedInclusive
 	return { Schema->ActionSchema.CreateNamedDiscreteInclusive({ Names, MakeArrayView(NormalizedPriorProbabilities) }, Tag) };
 }
 
-FLearningAgentsActionSchemaElement ULearningAgentsActions::SpecifyStructAction(ULearningAgentsActionSchema* Schema, const TMap<FName, FLearningAgentsActionSchemaElement>& Elements, const bool bActive, const FName Tag)
+FLearningAgentsActionSchemaElement ULearningAgentsActions::SpecifyStructAction(ULearningAgentsActionSchema* Schema, const TMap<FName, FLearningAgentsActionSchemaElement>& Elements, const FName Tag)
 {
-	if (!bActive)
-	{
-		return SpecifyNullAction(Schema, Tag);
-	}
-
 	if (Elements.Num() == 0)
 	{
 		UE_LOG(LogLearning, Warning, TEXT("SpecifyStructAction: Specifying zero-sized Struct Action."));
@@ -1425,26 +1375,16 @@ FLearningAgentsActionSchemaElement ULearningAgentsActions::SpecifyStructAction(U
 		SortedSubElements[Idx] = SubElements[SubElementIndices[Idx]];
 	}
 
-	return SpecifyStructActionFromArrayViews(Schema, SortedSubElementNames, SortedSubElements, bActive, Tag);
+	return SpecifyStructActionFromArrayViews(Schema, SortedSubElementNames, SortedSubElements, Tag);
 }
 
-FLearningAgentsActionSchemaElement ULearningAgentsActions::SpecifyStructActionFromArrays(ULearningAgentsActionSchema* Schema, const TArray<FName>& ElementNames, const TArray<FLearningAgentsActionSchemaElement>& Elements, const bool bActive, const FName Tag)
+FLearningAgentsActionSchemaElement ULearningAgentsActions::SpecifyStructActionFromArrays(ULearningAgentsActionSchema* Schema, const TArray<FName>& ElementNames, const TArray<FLearningAgentsActionSchemaElement>& Elements, const FName Tag)
 {
-	if (!bActive)
-	{
-		return SpecifyNullAction(Schema, Tag);
-	}
-
-	return SpecifyStructActionFromArrayViews(Schema, ElementNames, Elements, bActive, Tag);
+	return SpecifyStructActionFromArrayViews(Schema, ElementNames, Elements, Tag);
 }
 
-FLearningAgentsActionSchemaElement ULearningAgentsActions::SpecifyStructActionFromArrayViews(ULearningAgentsActionSchema* Schema, const TArrayView<const FName> ElementNames, const TArrayView<const FLearningAgentsActionSchemaElement> Elements, const bool bActive, const FName Tag)
+FLearningAgentsActionSchemaElement ULearningAgentsActions::SpecifyStructActionFromArrayViews(ULearningAgentsActionSchema* Schema, const TArrayView<const FName> ElementNames, const TArrayView<const FLearningAgentsActionSchemaElement> Elements, const FName Tag)
 {
-	if (!bActive)
-	{
-		return SpecifyNullAction(Schema, Tag);
-	}
-
 	if (!Schema)
 	{
 		UE_LOG(LogLearning, Error, TEXT("SpecifyStructActionFromArrayViews: Schema is nullptr."));
@@ -1485,13 +1425,8 @@ FLearningAgentsActionSchemaElement ULearningAgentsActions::SpecifyStructActionFr
 	return { Schema->ActionSchema.CreateAnd({ ElementNames, SubElements }, Tag) };
 }
 
-FLearningAgentsActionSchemaElement ULearningAgentsActions::SpecifyExclusiveUnionAction(ULearningAgentsActionSchema* Schema, const TMap<FName, FLearningAgentsActionSchemaElement>& Elements, const TMap<FName, float>& PriorProbabilities, const bool bActive, const FName Tag)
+FLearningAgentsActionSchemaElement ULearningAgentsActions::SpecifyExclusiveUnionAction(ULearningAgentsActionSchema* Schema, const TMap<FName, FLearningAgentsActionSchemaElement>& Elements, const TMap<FName, float>& PriorProbabilities, const FName Tag)
 {
-	if (!bActive)
-	{
-		return SpecifyNullAction(Schema, Tag);
-	}
-
 	if (Elements.Num() == 0)
 	{
 		UE_LOG(LogLearning, Warning, TEXT("SpecifyExclusiveUnionAction: Specifying zero-sized Exclusive Union Action."));
@@ -1539,26 +1474,16 @@ FLearningAgentsActionSchemaElement ULearningAgentsActions::SpecifyExclusiveUnion
 		SortedSubElementPriorProbabilities[Idx] = SubElementPriorProbabilities[SubElementIndices[Idx]];
 	}
 
-	return SpecifyExclusiveUnionActionFromArrayViews(Schema, SortedSubElementNames, SortedSubElements, SortedSubElementPriorProbabilities, bActive, Tag);
+	return SpecifyExclusiveUnionActionFromArrayViews(Schema, SortedSubElementNames, SortedSubElements, SortedSubElementPriorProbabilities, Tag);
 }
 
-FLearningAgentsActionSchemaElement ULearningAgentsActions::SpecifyExclusiveUnionActionFromArrays(ULearningAgentsActionSchema* Schema, const TArray<FName>& ElementNames, const TArray<FLearningAgentsActionSchemaElement>& Elements, const TArray<float>& PriorProbabilities, const bool bActive, const FName Tag)
+FLearningAgentsActionSchemaElement ULearningAgentsActions::SpecifyExclusiveUnionActionFromArrays(ULearningAgentsActionSchema* Schema, const TArray<FName>& ElementNames, const TArray<FLearningAgentsActionSchemaElement>& Elements, const TArray<float>& PriorProbabilities, const FName Tag)
 {
-	if (!bActive)
-	{
-		return SpecifyNullAction(Schema, Tag);
-	}
-
-	return SpecifyExclusiveUnionActionFromArrayViews(Schema, ElementNames, Elements, PriorProbabilities, bActive, Tag);
+	return SpecifyExclusiveUnionActionFromArrayViews(Schema, ElementNames, Elements, PriorProbabilities, Tag);
 }
 
-FLearningAgentsActionSchemaElement ULearningAgentsActions::SpecifyExclusiveUnionActionFromArrayViews(ULearningAgentsActionSchema* Schema, const TArrayView<const FName> ElementNames, const TArrayView<const FLearningAgentsActionSchemaElement> Elements, const TArrayView<const float> PriorProbabilities, const bool bActive, const FName Tag)
+FLearningAgentsActionSchemaElement ULearningAgentsActions::SpecifyExclusiveUnionActionFromArrayViews(ULearningAgentsActionSchema* Schema, const TArrayView<const FName> ElementNames, const TArrayView<const FLearningAgentsActionSchemaElement> Elements, const TArrayView<const float> PriorProbabilities, const FName Tag)
 {
-	if (!bActive)
-	{
-		return SpecifyNullAction(Schema, Tag);
-	}
-
 	if (!Schema)
 	{
 		UE_LOG(LogLearning, Error, TEXT("SpecifyExclusiveUnionActionFromArrayViews: Schema is nullptr."));
@@ -1613,13 +1538,8 @@ FLearningAgentsActionSchemaElement ULearningAgentsActions::SpecifyExclusiveUnion
 	return { Schema->ActionSchema.CreateOrExclusive({ ElementNames, SubElements, MakeArrayView(NormalizedPriorProbabilities) }, Tag)};
 }
 
-FLearningAgentsActionSchemaElement ULearningAgentsActions::SpecifyInclusiveUnionAction(ULearningAgentsActionSchema* Schema, const TMap<FName, FLearningAgentsActionSchemaElement>& Elements, const TMap<FName, float>& PriorProbabilities, const bool bActive, const FName Tag)
+FLearningAgentsActionSchemaElement ULearningAgentsActions::SpecifyInclusiveUnionAction(ULearningAgentsActionSchema* Schema, const TMap<FName, FLearningAgentsActionSchemaElement>& Elements, const TMap<FName, float>& PriorProbabilities, const FName Tag)
 {
-	if (!bActive)
-	{
-		return SpecifyNullAction(Schema, Tag);
-	}
-
 	if (Elements.Num() == 0)
 	{
 		UE_LOG(LogLearning, Warning, TEXT("SpecifyInclusiveUnionAction: Specifying zero-sized Inclusive Union Action."));
@@ -1667,26 +1587,16 @@ FLearningAgentsActionSchemaElement ULearningAgentsActions::SpecifyInclusiveUnion
 		SortedSubElementPriorProbabilities[Idx] = SubElementPriorProbabilities[SubElementIndices[Idx]];
 	}
 
-	return SpecifyInclusiveUnionActionFromArrayViews(Schema, SortedSubElementNames, SortedSubElements, SortedSubElementPriorProbabilities, bActive, Tag);
+	return SpecifyInclusiveUnionActionFromArrayViews(Schema, SortedSubElementNames, SortedSubElements, SortedSubElementPriorProbabilities, Tag);
 }
 
-FLearningAgentsActionSchemaElement ULearningAgentsActions::SpecifyInclusiveUnionActionFromArrays(ULearningAgentsActionSchema* Schema, const TArray<FName> ElementNames, const TArray<FLearningAgentsActionSchemaElement>& Elements, const TArray<float>& PriorProbabilities, const bool bActive, const FName Tag)
+FLearningAgentsActionSchemaElement ULearningAgentsActions::SpecifyInclusiveUnionActionFromArrays(ULearningAgentsActionSchema* Schema, const TArray<FName> ElementNames, const TArray<FLearningAgentsActionSchemaElement>& Elements, const TArray<float>& PriorProbabilities, const FName Tag)
 {
-	if (!bActive)
-	{
-		return SpecifyNullAction(Schema, Tag);
-	}
-
-	return SpecifyInclusiveUnionActionFromArrayViews(Schema, ElementNames, Elements, PriorProbabilities, bActive, Tag);
+	return SpecifyInclusiveUnionActionFromArrayViews(Schema, ElementNames, Elements, PriorProbabilities, Tag);
 }
 
-FLearningAgentsActionSchemaElement ULearningAgentsActions::SpecifyInclusiveUnionActionFromArrayViews(ULearningAgentsActionSchema* Schema, const TArrayView<const FName> ElementNames, const TArrayView<const FLearningAgentsActionSchemaElement> Elements, const TArrayView<const float> PriorProbabilities, const bool bActive, const FName Tag)
+FLearningAgentsActionSchemaElement ULearningAgentsActions::SpecifyInclusiveUnionActionFromArrayViews(ULearningAgentsActionSchema* Schema, const TArrayView<const FName> ElementNames, const TArrayView<const FLearningAgentsActionSchemaElement> Elements, const TArrayView<const float> PriorProbabilities, const FName Tag)
 {
-	if (!bActive)
-	{
-		return SpecifyNullAction(Schema, Tag);
-	}
-
 	if (!Schema)
 	{
 		UE_LOG(LogLearning, Error, TEXT("SpecifyInclusiveUnionActionFromArrayViews: Schema is nullptr."));
@@ -1741,13 +1651,8 @@ FLearningAgentsActionSchemaElement ULearningAgentsActions::SpecifyInclusiveUnion
 	return { Schema->ActionSchema.CreateOrInclusive({ ElementNames, SubElements, MakeArrayView(NormalizedPriorProbabilities) }, Tag)};
 }
 
-FLearningAgentsActionSchemaElement ULearningAgentsActions::SpecifyStaticArrayAction(ULearningAgentsActionSchema* Schema, const FLearningAgentsActionSchemaElement Element, const int32 Num, const bool bActive, const FName Tag)
+FLearningAgentsActionSchemaElement ULearningAgentsActions::SpecifyStaticArrayAction(ULearningAgentsActionSchema* Schema, const FLearningAgentsActionSchemaElement Element, const int32 Num, const FName Tag)
 {
-	if (!bActive)
-	{
-		return SpecifyNullAction(Schema, Tag);
-	}
-
 	if (!Schema)
 	{
 		UE_LOG(LogLearning, Error, TEXT("SpecifyStaticArrayAction: Schema is nullptr."));
@@ -1774,23 +1679,13 @@ FLearningAgentsActionSchemaElement ULearningAgentsActions::SpecifyStaticArrayAct
 	return { Schema->ActionSchema.CreateArray({ Element.SchemaElement, Num }, Tag) };
 }
 
-FLearningAgentsActionSchemaElement ULearningAgentsActions::SpecifyPairAction(ULearningAgentsActionSchema* Schema, const FLearningAgentsActionSchemaElement Key, const FLearningAgentsActionSchemaElement Value, const bool bActive, const FName Tag)
+FLearningAgentsActionSchemaElement ULearningAgentsActions::SpecifyPairAction(ULearningAgentsActionSchema* Schema, const FLearningAgentsActionSchemaElement Key, const FLearningAgentsActionSchemaElement Value, const FName Tag)
 {
-	if (!bActive)
-	{
-		return SpecifyNullAction(Schema, Tag);
-	}
-
-	return SpecifyStructActionFromArrayViews(Schema, { TEXT("Key"), TEXT("Value") }, { Key, Value }, bActive, Tag);
+	return SpecifyStructActionFromArrayViews(Schema, { TEXT("Key"), TEXT("Value") }, { Key, Value }, Tag);
 }
 
-FLearningAgentsActionSchemaElement ULearningAgentsActions::SpecifyEnumAction(ULearningAgentsActionSchema* Schema, const UEnum* Enum, const TMap<uint8, float>& PriorProbabilities, const bool bActive, const FName Tag)
+FLearningAgentsActionSchemaElement ULearningAgentsActions::SpecifyEnumAction(ULearningAgentsActionSchema* Schema, const UEnum* Enum, const TMap<uint8, float>& PriorProbabilities, const FName Tag)
 {
-	if (!bActive)
-	{
-		return SpecifyNullAction(Schema, Tag);
-	}
-
 	if (!Enum)
 	{
 		UE_LOG(LogLearning, Error, TEXT("SpecifyEnumAction: Enum is nullptr."));
@@ -1814,42 +1709,27 @@ FLearningAgentsActionSchemaElement ULearningAgentsActions::SpecifyEnumAction(ULe
 		}
 	}
 
-	return SpecifyEnumActionFromArrayView(Schema, Enum, PriorProbabilitiesArray, bActive, Tag);
+	return SpecifyEnumActionFromArrayView(Schema, Enum, PriorProbabilitiesArray, Tag);
 }
 
-FLearningAgentsActionSchemaElement ULearningAgentsActions::SpecifyEnumActionFromArray(ULearningAgentsActionSchema* Schema, const UEnum* Enum, const TArray<float>& PriorProbabilities, const bool bActive, const FName Tag)
+FLearningAgentsActionSchemaElement ULearningAgentsActions::SpecifyEnumActionFromArray(ULearningAgentsActionSchema* Schema, const UEnum* Enum, const TArray<float>& PriorProbabilities, const FName Tag)
 {
-	if (!bActive)
-	{
-		return SpecifyNullAction(Schema, Tag);
-	}
-
-	return SpecifyEnumActionFromArrayView(Schema, Enum, PriorProbabilities, bActive, Tag);
+	return SpecifyEnumActionFromArrayView(Schema, Enum, PriorProbabilities, Tag);
 }
 
-FLearningAgentsActionSchemaElement ULearningAgentsActions::SpecifyEnumActionFromArrayView(ULearningAgentsActionSchema* Schema, const UEnum* Enum, const TArrayView<const float> PriorProbabilities, const bool bActive, const FName Tag)
+FLearningAgentsActionSchemaElement ULearningAgentsActions::SpecifyEnumActionFromArrayView(ULearningAgentsActionSchema* Schema, const UEnum* Enum, const TArrayView<const float> PriorProbabilities, const FName Tag)
 {
-	if (!bActive)
-	{
-		return SpecifyNullAction(Schema, Tag);
-	}
-
 	if (!Enum)
 	{
 		UE_LOG(LogLearning, Error, TEXT("SpecifyEnumActionFromArrayView: Enum is nullptr."));
 		return FLearningAgentsActionSchemaElement();
 	}
 
-	return SpecifyExclusiveDiscreteActionFromArrayView(Schema, Enum->NumEnums() - 1, PriorProbabilities, bActive, Tag);
+	return SpecifyExclusiveDiscreteActionFromArrayView(Schema, Enum->NumEnums() - 1, PriorProbabilities, Tag);
 }
 
-FLearningAgentsActionSchemaElement ULearningAgentsActions::SpecifyBitmaskAction(ULearningAgentsActionSchema* Schema, const UEnum* Enum, const TMap<uint8, float>& PriorProbabilities, const bool bActive, const FName Tag)
+FLearningAgentsActionSchemaElement ULearningAgentsActions::SpecifyBitmaskAction(ULearningAgentsActionSchema* Schema, const UEnum* Enum, const TMap<uint8, float>& PriorProbabilities, const FName Tag)
 {
-	if (!bActive)
-	{
-		return SpecifyNullAction(Schema, Tag);
-	}
-
 	if (!Enum)
 	{
 		UE_LOG(LogLearning, Error, TEXT("SpecifyBitmaskAction: Enum is nullptr."));
@@ -1879,26 +1759,16 @@ FLearningAgentsActionSchemaElement ULearningAgentsActions::SpecifyBitmaskAction(
 		}
 	}
 
-	return SpecifyBitmaskActionFromArrayView(Schema, Enum, PriorProbabilitiesArray, bActive, Tag);
+	return SpecifyBitmaskActionFromArrayView(Schema, Enum, PriorProbabilitiesArray, Tag);
 }
 
-FLearningAgentsActionSchemaElement ULearningAgentsActions::SpecifyBitmaskActionFromArray(ULearningAgentsActionSchema* Schema, const UEnum* Enum, const TArray<float>& PriorProbabilities, const bool bActive, const FName Tag)
+FLearningAgentsActionSchemaElement ULearningAgentsActions::SpecifyBitmaskActionFromArray(ULearningAgentsActionSchema* Schema, const UEnum* Enum, const TArray<float>& PriorProbabilities, const FName Tag)
 {
-	if (!bActive)
-	{
-		return SpecifyNullAction(Schema, Tag);
-	}
-
-	return SpecifyBitmaskActionFromArrayView(Schema, Enum, PriorProbabilities, bActive, Tag);
+	return SpecifyBitmaskActionFromArrayView(Schema, Enum, PriorProbabilities, Tag);
 }
 
-FLearningAgentsActionSchemaElement ULearningAgentsActions::SpecifyBitmaskActionFromArrayView(ULearningAgentsActionSchema* Schema, const UEnum* Enum, const TArrayView<const float> PriorProbabilities, const bool bActive, const FName Tag)
+FLearningAgentsActionSchemaElement ULearningAgentsActions::SpecifyBitmaskActionFromArrayView(ULearningAgentsActionSchema* Schema, const UEnum* Enum, const TArrayView<const float> PriorProbabilities, const FName Tag)
 {
-	if (!bActive)
-	{
-		return SpecifyNullAction(Schema, Tag);
-	}
-
 	if (!Enum)
 	{
 		UE_LOG(LogLearning, Error, TEXT("SpecifyBitmaskActionFromArrayView: Enum is nullptr."));
@@ -1911,36 +1781,21 @@ FLearningAgentsActionSchemaElement ULearningAgentsActions::SpecifyBitmaskActionF
 		return FLearningAgentsActionSchemaElement();
 	}
 
-	return SpecifyInclusiveDiscreteActionFromArrayView(Schema, Enum->NumEnums() - 1, PriorProbabilities, bActive, Tag);
+	return SpecifyInclusiveDiscreteActionFromArrayView(Schema, Enum->NumEnums() - 1, PriorProbabilities, Tag);
 }
 
-FLearningAgentsActionSchemaElement ULearningAgentsActions::SpecifyOptionalAction(ULearningAgentsActionSchema* Schema, const FLearningAgentsActionSchemaElement Element, const float PriorProbability, const bool bActive, const FName Tag)
+FLearningAgentsActionSchemaElement ULearningAgentsActions::SpecifyOptionalAction(ULearningAgentsActionSchema* Schema, const FLearningAgentsActionSchemaElement Element, const float PriorProbability, const FName Tag)
 {
-	if (!bActive)
-	{
-		return SpecifyNullAction(Schema, Tag);
-	}
-
-	return SpecifyExclusiveUnionActionFromArrayViews(Schema, { TEXT("Null"), TEXT("Valid") }, { SpecifyNullAction(Schema), Element }, { 1.0f - PriorProbability, PriorProbability }, bActive, Tag);
+	return SpecifyExclusiveUnionActionFromArrayViews(Schema, { TEXT("Null"), TEXT("Valid") }, { SpecifyNullAction(Schema), Element }, { 1.0f - PriorProbability, PriorProbability }, Tag);
 }
 
-FLearningAgentsActionSchemaElement ULearningAgentsActions::SpecifyEitherAction(ULearningAgentsActionSchema* Schema, const FLearningAgentsActionSchemaElement A, const FLearningAgentsActionSchemaElement B, const float PriorProbabilityOfA, const bool bActive, const FName Tag)
+FLearningAgentsActionSchemaElement ULearningAgentsActions::SpecifyEitherAction(ULearningAgentsActionSchema* Schema, const FLearningAgentsActionSchemaElement A, const FLearningAgentsActionSchemaElement B, const float PriorProbabilityOfA, const FName Tag)
 {
-	if (!bActive)
-	{
-		return SpecifyNullAction(Schema, Tag);
-	}
-
-	return SpecifyExclusiveUnionActionFromArrayViews(Schema, { TEXT("A"), TEXT("B") }, { A, B }, { 1.0f - PriorProbabilityOfA, PriorProbabilityOfA }, bActive, Tag);
+	return SpecifyExclusiveUnionActionFromArrayViews(Schema, { TEXT("A"), TEXT("B") }, { A, B }, { 1.0f - PriorProbabilityOfA, PriorProbabilityOfA }, Tag);
 }
 
-FLearningAgentsActionSchemaElement ULearningAgentsActions::SpecifyEncodingAction(ULearningAgentsActionSchema* Schema, const FLearningAgentsActionSchemaElement Element, const int32 EncodingSize, const int32 HiddenLayerNum, const ELearningAgentsActivationFunction ActivationFunction, const bool bActive, const FName Tag)
+FLearningAgentsActionSchemaElement ULearningAgentsActions::SpecifyEncodingAction(ULearningAgentsActionSchema* Schema, const FLearningAgentsActionSchemaElement Element, const int32 EncodingSize, const int32 HiddenLayerNum, const ELearningAgentsActivationFunction ActivationFunction, const FName Tag)
 {
-	if (!bActive)
-	{
-		return SpecifyNullAction(Schema, Tag);
-	}
-
 	if (!Schema)
 	{
 		UE_LOG(LogLearning, Error, TEXT("SpecifyEncodingAction: Schema is nullptr."));
@@ -1968,54 +1823,29 @@ FLearningAgentsActionSchemaElement ULearningAgentsActions::SpecifyEncodingAction
 	return { Schema->ActionSchema.CreateEncoding({ Element.SchemaElement, EncodingSize, HiddenLayerNum, UE::Learning::Agents::Action::Private::GetEncodingActivationFunction(ActivationFunction) }, Tag)};
 }
 
-FLearningAgentsActionSchemaElement ULearningAgentsActions::SpecifyBoolAction(ULearningAgentsActionSchema* Schema, const float PriorProbability, const bool bActive, const FName Tag)
+FLearningAgentsActionSchemaElement ULearningAgentsActions::SpecifyBoolAction(ULearningAgentsActionSchema* Schema, const float PriorProbability, const FName Tag)
 {
-	if (!bActive)
-	{
-		return SpecifyNullAction(Schema, Tag);
-	}
-
-	return SpecifyExclusiveDiscreteActionFromArrayView(Schema, 2, { 1.0f - PriorProbability, PriorProbability }, bActive, Tag);
+	return SpecifyExclusiveDiscreteActionFromArrayView(Schema, 2, { 1.0f - PriorProbability, PriorProbability }, Tag);
 }
 
-FLearningAgentsActionSchemaElement ULearningAgentsActions::SpecifyFloatAction(ULearningAgentsActionSchema* Schema, const float FloatScale, const bool bActive, const FName Tag)
+FLearningAgentsActionSchemaElement ULearningAgentsActions::SpecifyFloatAction(ULearningAgentsActionSchema* Schema, const float FloatScale, const FName Tag)
 {
-	if (!bActive)
-	{
-		return SpecifyNullAction(Schema, Tag);
-	}
-
-	return SpecifyContinuousAction(Schema, 1, FloatScale, bActive, Tag);
+	return SpecifyContinuousAction(Schema, 1, FloatScale, Tag);
 }
 
-FLearningAgentsActionSchemaElement ULearningAgentsActions::SpecifyLocationAction(ULearningAgentsActionSchema* Schema, const float LocationScale, const bool bActive, const FName Tag)
+FLearningAgentsActionSchemaElement ULearningAgentsActions::SpecifyLocationAction(ULearningAgentsActionSchema* Schema, const float LocationScale, const FName Tag)
 {
-	if (!bActive)
-	{
-		return SpecifyNullAction(Schema, Tag);
-	}
-
-	return SpecifyContinuousAction(Schema, 3, LocationScale, bActive, Tag);
+	return SpecifyContinuousAction(Schema, 3, LocationScale, Tag);
 }
 
-FLearningAgentsActionSchemaElement ULearningAgentsActions::SpecifyRotationAction(ULearningAgentsActionSchema* Schema, const float RotationScale, const bool bActive, const FName Tag)
+FLearningAgentsActionSchemaElement ULearningAgentsActions::SpecifyRotationAction(ULearningAgentsActionSchema* Schema, const float RotationScale, const FName Tag)
 {
-	if (!bActive)
-	{
-		return SpecifyNullAction(Schema, Tag);
-	}
-
-	return SpecifyContinuousAction(Schema, 3, RotationScale, bActive, Tag);
+	return SpecifyContinuousAction(Schema, 3, RotationScale, Tag);
 }
 
-FLearningAgentsActionSchemaElement ULearningAgentsActions::SpecifyScaleAction(ULearningAgentsActionSchema* Schema, const float ScaleScale, const bool bActive, const FName Tag)
+FLearningAgentsActionSchemaElement ULearningAgentsActions::SpecifyScaleAction(ULearningAgentsActionSchema* Schema, const float ScaleScale, const FName Tag)
 {
-	if (!bActive)
-	{
-		return SpecifyNullAction(Schema, Tag);
-	}
-
-	return SpecifyContinuousAction(Schema, 3, ScaleScale, bActive, Tag);
+	return SpecifyContinuousAction(Schema, 3, ScaleScale, Tag);
 }
 
 FLearningAgentsActionSchemaElement ULearningAgentsActions::SpecifyTransformAction(
@@ -2023,56 +1853,35 @@ FLearningAgentsActionSchemaElement ULearningAgentsActions::SpecifyTransformActio
 	const float LocationScale ,
 	const float RotationScale,
 	const float ScaleScale,
-	const bool bActive,
 	const FName Tag)
 {
-	if (!bActive)
-	{
-		return SpecifyNullAction(Schema, Tag);
-	}
-
 	return SpecifyStructActionFromArrayViews(Schema,
 		{
 			TEXT("Location"),
 			TEXT("Rotation"),
 			TEXT("Scale")
 		},
-			{
-				SpecifyLocationAction(Schema, LocationScale),
-				SpecifyRotationAction(Schema, RotationScale),
-				SpecifyScaleAction(Schema, ScaleScale)
-			},
-		bActive, Tag);
+		{
+			SpecifyLocationAction(Schema, LocationScale),
+			SpecifyRotationAction(Schema, RotationScale),
+			SpecifyScaleAction(Schema, ScaleScale)
+		}, 
+		Tag);
 }
 
-FLearningAgentsActionSchemaElement ULearningAgentsActions::SpecifyAngleAction(ULearningAgentsActionSchema* Schema, const float AngleScale, const bool bActive, const FName Tag)
+FLearningAgentsActionSchemaElement ULearningAgentsActions::SpecifyAngleAction(ULearningAgentsActionSchema* Schema, const float AngleScale, const FName Tag)
 {
-	if (!bActive)
-	{
-		return SpecifyNullAction(Schema, Tag);
-	}
-
-	return SpecifyContinuousAction(Schema, 1, AngleScale, bActive, Tag);
+	return SpecifyContinuousAction(Schema, 1, AngleScale, Tag);
 }
 
-FLearningAgentsActionSchemaElement ULearningAgentsActions::SpecifyVelocityAction(ULearningAgentsActionSchema* Schema, const float VelocityScale, const bool bActive, const FName Tag)
+FLearningAgentsActionSchemaElement ULearningAgentsActions::SpecifyVelocityAction(ULearningAgentsActionSchema* Schema, const float VelocityScale, const FName Tag)
 {
-	if (!bActive)
-	{
-		return SpecifyNullAction(Schema, Tag);
-	}
-
-	return SpecifyContinuousAction(Schema, 3, VelocityScale, bActive, Tag);
+	return SpecifyContinuousAction(Schema, 3, VelocityScale, Tag);
 }
 
-FLearningAgentsActionSchemaElement ULearningAgentsActions::SpecifyDirectionAction(ULearningAgentsActionSchema* Schema, const bool bActive, const FName Tag)
+FLearningAgentsActionSchemaElement ULearningAgentsActions::SpecifyDirectionAction(ULearningAgentsActionSchema* Schema, const FName Tag)
 {
-	if (!bActive)
-	{
-		return SpecifyNullAction(Schema, Tag);
-	}
-
-	return SpecifyContinuousAction(Schema, 3, 1.0f, bActive, Tag);
+	return SpecifyContinuousAction(Schema, 3, 1.0f, Tag);
 }
 
 void ULearningAgentsActions::LogAction(const ULearningAgentsActionObject* Object, const FLearningAgentsActionObjectElement Element)
@@ -2100,7 +1909,6 @@ FLearningAgentsActionObjectElement ULearningAgentsActions::MakeNullAction(ULearn
 FLearningAgentsActionObjectElement ULearningAgentsActions::MakeContinuousAction(
 	ULearningAgentsActionObject* Object,
 	const TArray<float>& Values,
-    const bool bActive,	
 	const FName Tag,
 	const bool bVisualLoggerEnabled,
 	ULearningAgentsManagerListener* VisualLoggerListener,
@@ -2108,18 +1916,12 @@ FLearningAgentsActionObjectElement ULearningAgentsActions::MakeContinuousAction(
 	const FVector VisualLoggerLocation,
 	const FLinearColor VisualLoggerColor)
 {
-	if (!bActive)
-	{
-		return MakeNullAction(Object, Tag);
-	}
-
-	return MakeContinuousActionFromArrayView(Object, Values, bActive, Tag, bVisualLoggerEnabled, VisualLoggerListener, VisualLoggerAgentId, VisualLoggerLocation, VisualLoggerColor);
+	return MakeContinuousActionFromArrayView(Object, Values, Tag, bVisualLoggerEnabled, VisualLoggerListener, VisualLoggerAgentId, VisualLoggerLocation, VisualLoggerColor);
 }
 
 FLearningAgentsActionObjectElement ULearningAgentsActions::MakeContinuousActionFromArrayView(
 	ULearningAgentsActionObject* Object,
 	const TArrayView<const float> Values,
-    const bool bActive,
 	const FName Tag,
 	const bool bVisualLoggerEnabled,
 	ULearningAgentsManagerListener* VisualLoggerListener,
@@ -2127,11 +1929,6 @@ FLearningAgentsActionObjectElement ULearningAgentsActions::MakeContinuousActionF
 	const FVector VisualLoggerLocation,
 	const FLinearColor VisualLoggerColor)
 {
-	if (!bActive)
-	{
-		return MakeNullAction(Object, Tag);
-	}
-
 	if (!Object)
 	{
 		UE_LOG(LogLearning, Error, TEXT("MakeContinuousActionFromArrayView: Object is nullptr."));
@@ -2165,7 +1962,6 @@ FLearningAgentsActionObjectElement ULearningAgentsActions::MakeContinuousActionF
 FLearningAgentsActionObjectElement ULearningAgentsActions::MakeExclusiveDiscreteAction(
 	ULearningAgentsActionObject* Object,
 	const int32 Index,
-    const bool bActive,
 	const FName Tag,
 	const bool bVisualLoggerEnabled,
 	ULearningAgentsManagerListener* VisualLoggerListener,
@@ -2173,11 +1969,6 @@ FLearningAgentsActionObjectElement ULearningAgentsActions::MakeExclusiveDiscrete
 	const FVector VisualLoggerLocation,
 	const FLinearColor VisualLoggerColor)
 {
-	if (!bActive)
-	{
-		return MakeNullAction(Object, Tag);
-	}
-
 	if (!Object)
 	{
 		UE_LOG(LogLearning, Error, TEXT("MakeExclusiveDiscreteAction: Object is nullptr."));
@@ -2211,7 +2002,6 @@ FLearningAgentsActionObjectElement ULearningAgentsActions::MakeExclusiveDiscrete
 FLearningAgentsActionObjectElement ULearningAgentsActions::MakeNamedExclusiveDiscreteAction(
 	ULearningAgentsActionObject* Object,
 	const FName Name,
-    const bool bActive,	
 	const FName Tag,
 	const bool bVisualLoggerEnabled,
 	ULearningAgentsManagerListener* VisualLoggerListener,
@@ -2219,11 +2009,6 @@ FLearningAgentsActionObjectElement ULearningAgentsActions::MakeNamedExclusiveDis
 	const FVector VisualLoggerLocation,
 	const FLinearColor VisualLoggerColor)
 {
-	if (!bActive)
-	{
-		return MakeNullAction(Object, Tag);
-	}
-
 	if (!Object)
 	{
 		UE_LOG(LogLearning, Error, TEXT("MakeNamedExclusiveDiscreteAction: Object is nullptr."));
@@ -2251,7 +2036,6 @@ FLearningAgentsActionObjectElement ULearningAgentsActions::MakeNamedExclusiveDis
 FLearningAgentsActionObjectElement ULearningAgentsActions::MakeInclusiveDiscreteAction(
 	ULearningAgentsActionObject* Object,
 	const TArray<int32>& Indices,
-    const bool bActive,
 	const FName Tag,
 	const bool bVisualLoggerEnabled,
 	ULearningAgentsManagerListener* VisualLoggerListener,
@@ -2259,18 +2043,12 @@ FLearningAgentsActionObjectElement ULearningAgentsActions::MakeInclusiveDiscrete
 	const FVector VisualLoggerLocation,
 	const FLinearColor VisualLoggerColor)
 {
-	if (!bActive)
-	{
-		return MakeNullAction(Object, Tag);
-	}
-
-	return MakeInclusiveDiscreteActionFromArrayView(Object, Indices, bActive, Tag, bVisualLoggerEnabled, VisualLoggerListener, VisualLoggerAgentId, VisualLoggerLocation, VisualLoggerColor);
+	return MakeInclusiveDiscreteActionFromArrayView(Object, Indices, Tag, bVisualLoggerEnabled, VisualLoggerListener, VisualLoggerAgentId, VisualLoggerLocation, VisualLoggerColor);
 }
 
 FLearningAgentsActionObjectElement ULearningAgentsActions::MakeInclusiveDiscreteActionFromArrayView(
 	ULearningAgentsActionObject* Object,
 	const TArrayView<const int32> Indices,
-    const bool bActive,	
 	const FName Tag,
 	const bool bVisualLoggerEnabled,
 	ULearningAgentsManagerListener* VisualLoggerListener,
@@ -2278,11 +2056,6 @@ FLearningAgentsActionObjectElement ULearningAgentsActions::MakeInclusiveDiscrete
 	const FVector VisualLoggerLocation,
 	const FLinearColor VisualLoggerColor)
 {
-	if (!bActive)
-	{
-		return MakeNullAction(Object, Tag);
-	}
-
 	if (!Object)
 	{
 		UE_LOG(LogLearning, Error, TEXT("MakeInclusiveDiscreteActionFromArrayView: Object is nullptr."));
@@ -2327,7 +2100,6 @@ FLearningAgentsActionObjectElement ULearningAgentsActions::MakeInclusiveDiscrete
 FLearningAgentsActionObjectElement ULearningAgentsActions::MakeNamedInclusiveDiscreteAction(
 	ULearningAgentsActionObject* Object,
 	const TArray<FName>& Names,
-    const bool bActive,	
 	const FName Tag,
 	const bool bVisualLoggerEnabled,
 	ULearningAgentsManagerListener* VisualLoggerListener,
@@ -2335,18 +2107,12 @@ FLearningAgentsActionObjectElement ULearningAgentsActions::MakeNamedInclusiveDis
 	const FVector VisualLoggerLocation,
 	const FLinearColor VisualLoggerColor)
 {
-	if (!bActive)
-	{
-		return MakeNullAction(Object, Tag);
-	}
-
-	return MakeNamedInclusiveDiscreteActionFromArrayView(Object, Names, bActive, Tag, bVisualLoggerEnabled, VisualLoggerListener, VisualLoggerAgentId, VisualLoggerLocation, VisualLoggerColor);
+	return MakeNamedInclusiveDiscreteActionFromArrayView(Object, Names, Tag, bVisualLoggerEnabled, VisualLoggerListener, VisualLoggerAgentId, VisualLoggerLocation, VisualLoggerColor);
 }
 
 FLearningAgentsActionObjectElement ULearningAgentsActions::MakeNamedInclusiveDiscreteActionFromArrayView(
 	ULearningAgentsActionObject* Object,
 	const TArrayView<const FName> Names,
-    const bool bActive,
 	const FName Tag,
 	const bool bVisualLoggerEnabled,
 	ULearningAgentsManagerListener* VisualLoggerListener,
@@ -2354,11 +2120,6 @@ FLearningAgentsActionObjectElement ULearningAgentsActions::MakeNamedInclusiveDis
 	const FVector VisualLoggerLocation,
 	const FLinearColor VisualLoggerColor)
 {
-	if (!bActive)
-	{
-		return MakeNullAction(Object, Tag);
-	}
-
 	if (!Object)
 	{
 		UE_LOG(LogLearning, Error, TEXT("MakeNamedInclusiveDiscreteActionFromArrayView: Object is nullptr."));
@@ -2389,13 +2150,8 @@ FLearningAgentsActionObjectElement ULearningAgentsActions::MakeNamedInclusiveDis
 	return { Object->ActionObject.CreateNamedDiscreteInclusive({ Names }, Tag) };
 }
 
-FLearningAgentsActionObjectElement ULearningAgentsActions::MakeStructAction(ULearningAgentsActionObject* Object, const TMap<FName, FLearningAgentsActionObjectElement>& Elements, const bool bActive, const FName Tag)
+FLearningAgentsActionObjectElement ULearningAgentsActions::MakeStructAction(ULearningAgentsActionObject* Object, const TMap<FName, FLearningAgentsActionObjectElement>& Elements, const FName Tag)
 {
-	if (!bActive)
-	{
-		return MakeNullAction(Object, Tag);
-	}
-
 	if (Elements.Num() == 0)
 	{
 		UE_LOG(LogLearning, Warning, TEXT("MakeStructAction: Creating zero-sized Struct Action."));
@@ -2414,26 +2170,16 @@ FLearningAgentsActionObjectElement ULearningAgentsActions::MakeStructAction(ULea
 		SubElementNames.Add(Element.Key);
 	}
 
-	return MakeStructActionFromArrayViews(Object, SubElementNames, SubElements, bActive, Tag);
+	return MakeStructActionFromArrayViews(Object, SubElementNames, SubElements, Tag);
 }
 
-FLearningAgentsActionObjectElement ULearningAgentsActions::MakeStructActionFromArrays(ULearningAgentsActionObject* Object, const TArray<FName>& ElementNames, const TArray<FLearningAgentsActionObjectElement>& Elements, const bool bActive, const FName Tag)
+FLearningAgentsActionObjectElement ULearningAgentsActions::MakeStructActionFromArrays(ULearningAgentsActionObject* Object, const TArray<FName>& ElementNames, const TArray<FLearningAgentsActionObjectElement>& Elements, const FName Tag)
 {
-	if (!bActive)
-	{
-		return MakeNullAction(Object, Tag);
-	}
-
-	return MakeStructActionFromArrayViews(Object, ElementNames, Elements, bActive, Tag);
+	return MakeStructActionFromArrayViews(Object, ElementNames, Elements, Tag);
 }
 
-FLearningAgentsActionObjectElement ULearningAgentsActions::MakeStructActionFromArrayViews(ULearningAgentsActionObject* Object, const TArrayView<const FName> ElementNames, const TArrayView<const FLearningAgentsActionObjectElement> Elements, const bool bActive, const FName Tag)
+FLearningAgentsActionObjectElement ULearningAgentsActions::MakeStructActionFromArrayViews(ULearningAgentsActionObject* Object, const TArrayView<const FName> ElementNames, const TArrayView<const FLearningAgentsActionObjectElement> Elements, const FName Tag)
 {
-	if (!bActive)
-	{
-		return MakeNullAction(Object, Tag);
-	}
-
 	if (!Object)
 	{
 		UE_LOG(LogLearning, Error, TEXT("MakeStructActionFromArrayViews: Object is nullptr."));
@@ -2474,13 +2220,8 @@ FLearningAgentsActionObjectElement ULearningAgentsActions::MakeStructActionFromA
 	return { Object->ActionObject.CreateAnd({ ElementNames, SubElements }, Tag) };
 }
 
-FLearningAgentsActionObjectElement ULearningAgentsActions::MakeExclusiveUnionAction(ULearningAgentsActionObject* Object, const FName ElementName, const FLearningAgentsActionObjectElement Element, const bool bActive, const FName Tag)
+FLearningAgentsActionObjectElement ULearningAgentsActions::MakeExclusiveUnionAction(ULearningAgentsActionObject* Object, const FName ElementName, const FLearningAgentsActionObjectElement Element, const FName Tag)
 {
-	if (!bActive)
-	{
-		return MakeNullAction(Object, Tag);
-	}
-
 	if (!Object)
 	{
 		UE_LOG(LogLearning, Error, TEXT("MakeExclusiveUnionAction: Object is nullptr."));
@@ -2496,13 +2237,8 @@ FLearningAgentsActionObjectElement ULearningAgentsActions::MakeExclusiveUnionAct
 	return { Object->ActionObject.CreateOrExclusive({ ElementName, Element.ObjectElement }, Tag) };
 }
 
-FLearningAgentsActionObjectElement ULearningAgentsActions::MakeInclusiveUnionAction(ULearningAgentsActionObject* Object, const TMap<FName, FLearningAgentsActionObjectElement>& Elements, const bool bActive, const FName Tag)
+FLearningAgentsActionObjectElement ULearningAgentsActions::MakeInclusiveUnionAction(ULearningAgentsActionObject* Object, const TMap<FName, FLearningAgentsActionObjectElement>& Elements, const FName Tag)
 {
-	if (!bActive)
-	{
-		return MakeNullAction(Object, Tag);
-	}
-
 	const int32 SubElementNum = Elements.Num();
 
 	TArray<FLearningAgentsActionObjectElement, TInlineAllocator<16>> SubElements;
@@ -2516,26 +2252,16 @@ FLearningAgentsActionObjectElement ULearningAgentsActions::MakeInclusiveUnionAct
 		SubElementNames.Add(Element.Key);
 	}
 
-	return MakeInclusiveUnionActionFromArrayViews(Object, SubElementNames, SubElements, bActive, Tag);
+	return MakeInclusiveUnionActionFromArrayViews(Object, SubElementNames, SubElements, Tag);
 }
 
-FLearningAgentsActionObjectElement ULearningAgentsActions::MakeInclusiveUnionActionFromArrays(ULearningAgentsActionObject* Object, const TArray<FName>& ElementNames, const TArray<FLearningAgentsActionObjectElement>& Elements, const bool bActive, const FName Tag)
+FLearningAgentsActionObjectElement ULearningAgentsActions::MakeInclusiveUnionActionFromArrays(ULearningAgentsActionObject* Object, const TArray<FName>& ElementNames, const TArray<FLearningAgentsActionObjectElement>& Elements, const FName Tag)
 {
-	if (!bActive)
-	{
-		return MakeNullAction(Object, Tag);
-	}
-
-	return MakeInclusiveUnionActionFromArrayViews(Object, ElementNames, Elements, bActive, Tag);
+	return MakeInclusiveUnionActionFromArrayViews(Object, ElementNames, Elements, Tag);
 }
 
-FLearningAgentsActionObjectElement ULearningAgentsActions::MakeInclusiveUnionActionFromArrayViews(ULearningAgentsActionObject* Object, const TArrayView<const FName> ElementNames, const TArrayView<const FLearningAgentsActionObjectElement> Elements, const bool bActive, const FName Tag)
+FLearningAgentsActionObjectElement ULearningAgentsActions::MakeInclusiveUnionActionFromArrayViews(ULearningAgentsActionObject* Object, const TArrayView<const FName> ElementNames, const TArrayView<const FLearningAgentsActionObjectElement> Elements, const FName Tag)
 {
-	if (!bActive)
-	{
-		return MakeNullAction(Object, Tag);
-	}
-
 	if (!Object)
 	{
 		UE_LOG(LogLearning, Error, TEXT("MakeInclusiveUnionActionFromArrayViews: Object is nullptr."));
@@ -2571,23 +2297,13 @@ FLearningAgentsActionObjectElement ULearningAgentsActions::MakeInclusiveUnionAct
 	return { Object->ActionObject.CreateOrInclusive({ ElementNames, SubElements }, Tag) };
 }
 
-FLearningAgentsActionObjectElement ULearningAgentsActions::MakeStaticArrayAction(ULearningAgentsActionObject* Object, const TArray<FLearningAgentsActionObjectElement>& Elements, const bool bActive, const FName Tag)
+FLearningAgentsActionObjectElement ULearningAgentsActions::MakeStaticArrayAction(ULearningAgentsActionObject* Object, const TArray<FLearningAgentsActionObjectElement>& Elements, const FName Tag)
 {
-	if (!bActive)
-	{
-		return MakeNullAction(Object, Tag);
-	}
-
-	return MakeStaticArrayActionFromArrayView(Object, Elements, bActive, Tag);
+	return MakeStaticArrayActionFromArrayView(Object, Elements, Tag);
 }
 
-FLearningAgentsActionObjectElement ULearningAgentsActions::MakeStaticArrayActionFromArrayView(ULearningAgentsActionObject* Object, const TArrayView<const FLearningAgentsActionObjectElement> Elements, const bool bActive, const FName Tag)
+FLearningAgentsActionObjectElement ULearningAgentsActions::MakeStaticArrayActionFromArrayView(ULearningAgentsActionObject* Object, const TArrayView<const FLearningAgentsActionObjectElement> Elements, const FName Tag)
 {
-	if (!bActive)
-	{
-		return MakeNullAction(Object, Tag);
-	}
-
 	if (!Object)
 	{
 		UE_LOG(LogLearning, Error, TEXT("MakeStaticArrayActionFromArrayView: Object is nullptr."));
@@ -2616,21 +2332,15 @@ FLearningAgentsActionObjectElement ULearningAgentsActions::MakeStaticArrayAction
 	return { Object->ActionObject.CreateArray({ SubElements }, Tag) };
 }
 
-FLearningAgentsActionObjectElement ULearningAgentsActions::MakePairAction(ULearningAgentsActionObject* Object, const FLearningAgentsActionObjectElement Key, const FLearningAgentsActionObjectElement Value, const bool bActive, const FName Tag)
+FLearningAgentsActionObjectElement ULearningAgentsActions::MakePairAction(ULearningAgentsActionObject* Object, const FLearningAgentsActionObjectElement Key, const FLearningAgentsActionObjectElement Value, const FName Tag)
 {
-	if (!bActive)
-	{
-		return MakeNullAction(Object, Tag);
-	}
-
-	return MakeStructActionFromArrayViews(Object, { TEXT("Key"), TEXT("Value") }, { Key, Value }, bActive, Tag);
+	return MakeStructActionFromArrayViews(Object, { TEXT("Key"), TEXT("Value") }, { Key, Value }, Tag);
 }
 
 FLearningAgentsActionObjectElement ULearningAgentsActions::MakeEnumAction(
 	ULearningAgentsActionObject* Object,
 	const UEnum* Enum,
 	const uint8 EnumValue,
-    const bool bActive,
 	const FName Tag,
 	const bool bVisualLoggerEnabled,
 	ULearningAgentsManagerListener* VisualLoggerListener,
@@ -2638,11 +2348,6 @@ FLearningAgentsActionObjectElement ULearningAgentsActions::MakeEnumAction(
 	const FVector VisualLoggerLocation,
 	const FLinearColor VisualLoggerColor)
 {
-	if (!bActive)
-	{
-		return MakeNullAction(Object, Tag);
-	}
-
 	if (!Enum)
 	{
 		UE_LOG(LogLearning, Error, TEXT("MakeEnumAction: Enum is nullptr."));
@@ -2675,14 +2380,13 @@ FLearningAgentsActionObjectElement ULearningAgentsActions::MakeEnumAction(
 	}
 #endif
 
-	return MakeExclusiveDiscreteAction(Object, EnumValueIndex, bActive, Tag);
+	return MakeExclusiveDiscreteAction(Object, EnumValueIndex, Tag);
 }
 
 FLearningAgentsActionObjectElement ULearningAgentsActions::MakeBitmaskAction(
 	ULearningAgentsActionObject* Object,
 	const UEnum* Enum,
 	const int32 BitmaskValue,
-    const bool bActive,	
 	const FName Tag,
 	const bool bVisualLoggerEnabled,
 	ULearningAgentsManagerListener* VisualLoggerListener,
@@ -2690,11 +2394,6 @@ FLearningAgentsActionObjectElement ULearningAgentsActions::MakeBitmaskAction(
 	const FVector VisualLoggerLocation,
 	const FLinearColor VisualLoggerColor)
 {
-	if (!bActive)
-	{
-		return MakeNullAction(Object, Tag);
-	}
-
 	if (!Enum)
 	{
 		UE_LOG(LogLearning, Error, TEXT("MakeBitmaskAction: Enum is nullptr."));
@@ -2751,79 +2450,44 @@ FLearningAgentsActionObjectElement ULearningAgentsActions::MakeBitmaskAction(
 	}
 #endif
 
-	return MakeInclusiveDiscreteActionFromArrayView(Object, BitmaskIndices, bActive, Tag);
+	return MakeInclusiveDiscreteActionFromArrayView(Object, BitmaskIndices, Tag);
 }
 
-FLearningAgentsActionObjectElement ULearningAgentsActions::MakeOptionalAction(ULearningAgentsActionObject* Object, const FLearningAgentsActionObjectElement Element, const ELearningAgentsOptionalAction Option, const bool bActive, const FName Tag)
+FLearningAgentsActionObjectElement ULearningAgentsActions::MakeOptionalAction(ULearningAgentsActionObject* Object, const FLearningAgentsActionObjectElement Element, const ELearningAgentsOptionalAction Option, const FName Tag)
 {
-	if (!bActive)
-	{
-		return MakeNullAction(Object, Tag);
-	}
-
 	return MakeExclusiveUnionAction(Object,
 		Option == ELearningAgentsOptionalAction::Null ? TEXT("Null") : TEXT("Valid"),
 		Option == ELearningAgentsOptionalAction::Null ? MakeNullAction(Object) : Element,
-		bActive, Tag);
+		Tag);
 }
 
-FLearningAgentsActionObjectElement ULearningAgentsActions::MakeOptionalNullAction(ULearningAgentsActionObject* Object, const bool bActive, const FName Tag)
+FLearningAgentsActionObjectElement ULearningAgentsActions::MakeOptionalNullAction(ULearningAgentsActionObject* Object, const FName Tag)
 {
-	if (!bActive)
-	{
-		return MakeNullAction(Object, Tag);
-	}
-
-	return MakeExclusiveUnionAction(Object, TEXT("Null"), MakeNullAction(Object), bActive, Tag);
+	return MakeExclusiveUnionAction(Object, TEXT("Null"), MakeNullAction(Object), Tag);
 }
 
-FLearningAgentsActionObjectElement ULearningAgentsActions::MakeOptionalValidAction(ULearningAgentsActionObject* Object, const FLearningAgentsActionObjectElement Element, const bool bActive, const FName Tag)
+FLearningAgentsActionObjectElement ULearningAgentsActions::MakeOptionalValidAction(ULearningAgentsActionObject* Object, const FLearningAgentsActionObjectElement Element, const FName Tag)
 {
-	if (!bActive)
-	{
-		return MakeNullAction(Object, Tag);
-	}
-
-	return MakeExclusiveUnionAction(Object, TEXT("Valid"), Element, bActive, Tag);
+	return MakeExclusiveUnionAction(Object, TEXT("Valid"), Element, Tag);
 }
 
-FLearningAgentsActionObjectElement ULearningAgentsActions::MakeEitherAction(ULearningAgentsActionObject* Object, const FLearningAgentsActionObjectElement Element, const ELearningAgentsEitherAction Either, const bool bActive, const FName Tag)
+FLearningAgentsActionObjectElement ULearningAgentsActions::MakeEitherAction(ULearningAgentsActionObject* Object, const FLearningAgentsActionObjectElement Element, const ELearningAgentsEitherAction Either, const FName Tag)
 {
-	if (!bActive)
-	{
-		return MakeNullAction(Object, Tag);
-	}
-
-	return MakeExclusiveUnionAction(Object, Either == ELearningAgentsEitherAction::A ? TEXT("A") : TEXT("B"), Element, bActive, Tag);
+	return MakeExclusiveUnionAction(Object, Either == ELearningAgentsEitherAction::A ? TEXT("A") : TEXT("B"), Element, Tag);
 }
 
-FLearningAgentsActionObjectElement ULearningAgentsActions::MakeEitherAAction(ULearningAgentsActionObject* Object, const FLearningAgentsActionObjectElement A, const bool bActive, const FName Tag)
+FLearningAgentsActionObjectElement ULearningAgentsActions::MakeEitherAAction(ULearningAgentsActionObject* Object, const FLearningAgentsActionObjectElement A, const FName Tag)
 {
-	if (!bActive)
-	{
-		return MakeNullAction(Object, Tag);
-	}
-
-	return MakeExclusiveUnionAction(Object, TEXT("A"), A, bActive, Tag);
+	return MakeExclusiveUnionAction(Object, TEXT("A"), A, Tag);
 }
 
-FLearningAgentsActionObjectElement ULearningAgentsActions::MakeEitherBAction(ULearningAgentsActionObject* Object, const FLearningAgentsActionObjectElement B, const bool bActive, const FName Tag)
+FLearningAgentsActionObjectElement ULearningAgentsActions::MakeEitherBAction(ULearningAgentsActionObject* Object, const FLearningAgentsActionObjectElement B, const FName Tag)
 {
-	if (!bActive)
-	{
-		return MakeNullAction(Object, Tag);
-	}
-
-	return MakeExclusiveUnionAction(Object, TEXT("B"), B, bActive, Tag);
+	return MakeExclusiveUnionAction(Object, TEXT("B"), B, Tag);
 }
 
-FLearningAgentsActionObjectElement ULearningAgentsActions::MakeEncodingAction(ULearningAgentsActionObject* Object, const FLearningAgentsActionObjectElement Element, const bool bActive, const FName Tag)
+FLearningAgentsActionObjectElement ULearningAgentsActions::MakeEncodingAction(ULearningAgentsActionObject* Object, const FLearningAgentsActionObjectElement Element, const FName Tag)
 {
-	if (!bActive)
-	{
-		return MakeNullAction(Object, Tag);
-	}
-
 	if (!Object)
 	{
 		UE_LOG(LogLearning, Error, TEXT("MakeEncodingAction: Object is nullptr."));
@@ -2842,7 +2506,6 @@ FLearningAgentsActionObjectElement ULearningAgentsActions::MakeEncodingAction(UL
 FLearningAgentsActionObjectElement ULearningAgentsActions::MakeBoolAction(
 	ULearningAgentsActionObject* Object,
 	const bool bValue,
-    const bool bActive,
 	const FName Tag,
 	const bool bVisualLoggerEnabled,
 	ULearningAgentsManagerListener* VisualLoggerListener,
@@ -2850,11 +2513,6 @@ FLearningAgentsActionObjectElement ULearningAgentsActions::MakeBoolAction(
 	const FVector VisualLoggerLocation,
 	const FLinearColor VisualLoggerColor)
 {
-	if (!bActive)
-	{
-		return MakeNullAction(Object, Tag);
-	}
-
 #if UE_LEARNING_AGENTS_ENABLE_VISUAL_LOG
 	if (bVisualLoggerEnabled && VisualLoggerListener)
 	{
@@ -2870,13 +2528,12 @@ FLearningAgentsActionObjectElement ULearningAgentsActions::MakeBoolAction(
 	}
 #endif
 
-	return MakeExclusiveDiscreteAction(Object, bValue ? 1 : 0, bActive, Tag);
+	return MakeExclusiveDiscreteAction(Object, bValue ? 1 : 0, Tag);
 }
 
 FLearningAgentsActionObjectElement ULearningAgentsActions::MakeFloatAction(
 	ULearningAgentsActionObject* Object,
 	const float Value,
-    const bool bActive,
 	const FName Tag,
 	const bool bVisualLoggerEnabled,
 	ULearningAgentsManagerListener* VisualLoggerListener,
@@ -2884,11 +2541,6 @@ FLearningAgentsActionObjectElement ULearningAgentsActions::MakeFloatAction(
 	const FVector VisualLoggerLocation,
 	const FLinearColor VisualLoggerColor)
 {
-	if (!bActive)
-	{
-		return MakeNullAction(Object, Tag);
-	}
-
 #if UE_LEARNING_AGENTS_ENABLE_VISUAL_LOG
 	if (bVisualLoggerEnabled && VisualLoggerListener)
 	{
@@ -2904,14 +2556,13 @@ FLearningAgentsActionObjectElement ULearningAgentsActions::MakeFloatAction(
 	}
 #endif
 
-	return MakeContinuousActionFromArrayView(Object, { Value }, bActive, Tag);
+	return MakeContinuousActionFromArrayView(Object, { Value }, Tag);
 }
 
 FLearningAgentsActionObjectElement ULearningAgentsActions::MakeLocationAction(
 	ULearningAgentsActionObject* Object,
 	const FVector Location,
 	const FTransform RelativeTransform,
-	const bool bActive,	
 	const FName Tag,
 	const bool bVisualLoggerEnabled,
 	ULearningAgentsManagerListener* VisualLoggerListener,
@@ -2919,11 +2570,6 @@ FLearningAgentsActionObjectElement ULearningAgentsActions::MakeLocationAction(
 	const FVector VisualLoggerLocation,
 	const FLinearColor VisualLoggerColor)
 {
-	if (!bActive)
-	{
-		return MakeNullAction(Object, Tag);
-	}
-
 	const FVector LocalLocation = RelativeTransform.InverseTransformPosition(Location);
 
 #if UE_LEARNING_AGENTS_ENABLE_VISUAL_LOG
@@ -2963,14 +2609,13 @@ FLearningAgentsActionObjectElement ULearningAgentsActions::MakeLocationAction(
 	return MakeContinuousActionFromArrayView(Object, {
 		(float)LocalLocation.X,
 		(float)LocalLocation.Y,
-		(float)LocalLocation.Z }, bActive, Tag);
+		(float)LocalLocation.Z }, Tag);
 }
 
 FLearningAgentsActionObjectElement ULearningAgentsActions::MakeRotationAction(
 	ULearningAgentsActionObject* Object,
 	const FRotator Rotation,
 	const FRotator RelativeRotation,
-    const bool bActive,	
 	const FName Tag,
 	const bool bVisualLoggerEnabled,
 	ULearningAgentsManagerListener* VisualLoggerListener,
@@ -2979,12 +2624,7 @@ FLearningAgentsActionObjectElement ULearningAgentsActions::MakeRotationAction(
 	const FVector VisualLoggerLocation,
 	const FLinearColor VisualLoggerColor)
 {
-	if (!bActive)
-	{
-		return MakeNullAction(Object, Tag);
-	}
-
-	return MakeRotationActionFromQuat(Object, FQuat::MakeFromRotator(Rotation), FQuat::MakeFromRotator(RelativeRotation), bActive, Tag,
+	return MakeRotationActionFromQuat(Object, FQuat::MakeFromRotator(Rotation), FQuat::MakeFromRotator(RelativeRotation), Tag,
 		bVisualLoggerEnabled, VisualLoggerListener, VisualLoggerAgentId, VisualLoggerRotationLocation, VisualLoggerLocation, VisualLoggerColor);
 }
 
@@ -2992,7 +2632,6 @@ FLearningAgentsActionObjectElement ULearningAgentsActions::MakeRotationActionFro
 	ULearningAgentsActionObject* Object,
 	const FQuat Rotation,
 	const FQuat RelativeRotation,
-	const bool bActive,
 	const FName Tag,
 	const bool bVisualLoggerEnabled,
 	ULearningAgentsManagerListener* VisualLoggerListener,
@@ -3001,11 +2640,6 @@ FLearningAgentsActionObjectElement ULearningAgentsActions::MakeRotationActionFro
 	const FVector VisualLoggerLocation,
 	const FLinearColor VisualLoggerColor)
 {
-	if (!bActive)
-	{
-		return MakeNullAction(Object, Tag);
-	}
-
 	FQuat LocalRotation = (RelativeRotation.Inverse() * Rotation).GetShortestArcWith(FQuat::Identity);
 	const FVector RotationVector = LocalRotation.ToRotationVector();
 
@@ -3036,14 +2670,13 @@ FLearningAgentsActionObjectElement ULearningAgentsActions::MakeRotationActionFro
 		(float)FMath::RadiansToDegrees(RotationVector.X),
 		(float)FMath::RadiansToDegrees(RotationVector.Y),
 		(float)FMath::RadiansToDegrees(RotationVector.Z),
-		}, bActive, Tag);
+		}, Tag);
 }
 
 FLearningAgentsActionObjectElement ULearningAgentsActions::MakeScaleAction(
 	ULearningAgentsActionObject* Object,
 	const FVector Scale,
 	const FVector RelativeScale,
-    const bool bActive,	
 	const FName Tag,
 	const bool bVisualLoggerEnabled,
 	ULearningAgentsManagerListener* VisualLoggerListener,
@@ -3051,11 +2684,6 @@ FLearningAgentsActionObjectElement ULearningAgentsActions::MakeScaleAction(
 	const FVector VisualLoggerLocation,
 	const FLinearColor VisualLoggerColor)
 {
-	if (!bActive)
-	{
-		return MakeNullAction(Object, Tag);
-	}
-
 	const FVector LocalLogScale =
 		UE::Learning::Agents::Action::Private::VectorLogSafe(Scale) -
 		UE::Learning::Agents::Action::Private::VectorLogSafe(RelativeScale);
@@ -3080,14 +2708,13 @@ FLearningAgentsActionObjectElement ULearningAgentsActions::MakeScaleAction(
 		(float)LocalLogScale.X,
 		(float)LocalLogScale.Y,
 		(float)LocalLogScale.Z,
-		}, bActive, Tag);
+		}, Tag);
 }
 
 FLearningAgentsActionObjectElement ULearningAgentsActions::MakeTransformAction(
 	ULearningAgentsActionObject* Object,
 	const FTransform Transform,
 	const FTransform RelativeTransform,
-	const bool bActive,	
 	const FName Tag,
 	const bool bVisualLoggerEnabled,
 	ULearningAgentsManagerListener* VisualLoggerListener,
@@ -3095,11 +2722,6 @@ FLearningAgentsActionObjectElement ULearningAgentsActions::MakeTransformAction(
 	const FVector VisualLoggerLocation,
 	const FLinearColor VisualLoggerColor)
 {
-	if (!bActive)
-	{
-		return MakeNullAction(Object, Tag);
-	}
-
 	FVector Location = Transform.GetLocation();
 	FQuat Rotation = Transform.GetRotation();
 	FVector Scale = Transform.GetScale3D();
@@ -3138,14 +2760,13 @@ FLearningAgentsActionObjectElement ULearningAgentsActions::MakeTransformAction(
 			MakeRotationActionFromQuat(Object, Rotation, RelativeTransform.GetRotation()),
 			MakeScaleAction(Object, Scale, RelativeTransform.GetScale3D())
 		},
-		bActive, Tag);
+		Tag);
 }
 
 FLearningAgentsActionObjectElement ULearningAgentsActions::MakeAngleAction(
 	ULearningAgentsActionObject* Object,
 	const float Angle,
 	const float RelativeAngle,
-	const bool bActive,
 	const FName Tag,
 	const bool bVisualLoggerEnabled,
 	ULearningAgentsManagerListener* VisualLoggerListener,
@@ -3154,11 +2775,6 @@ FLearningAgentsActionObjectElement ULearningAgentsActions::MakeAngleAction(
 	const FVector VisualLoggerLocation,
 	const FLinearColor VisualLoggerColor)
 {
-	if (!bActive)
-	{
-		return MakeNullAction(Object, Tag);
-	}
-
 	float LocalAngle = FMath::FindDeltaAngleDegrees(RelativeAngle, Angle);
 	
 #if UE_LEARNING_AGENTS_ENABLE_VISUAL_LOG
@@ -3185,14 +2801,13 @@ FLearningAgentsActionObjectElement ULearningAgentsActions::MakeAngleAction(
 	}
 #endif
 
-	return MakeContinuousActionFromArrayView(Object, { LocalAngle }, bActive, Tag);
+	return MakeContinuousActionFromArrayView(Object, { LocalAngle }, Tag);
 }
 
 FLearningAgentsActionObjectElement ULearningAgentsActions::MakeAngleActionRadians(
 	ULearningAgentsActionObject* Object,
 	const float Angle,
 	const float RelativeAngle,
-    const bool bActive,	
 	const FName Tag,
 	const bool bVisualLoggerEnabled,
 	ULearningAgentsManagerListener* VisualLoggerListener,
@@ -3201,12 +2816,7 @@ FLearningAgentsActionObjectElement ULearningAgentsActions::MakeAngleActionRadian
 	const FVector VisualLoggerLocation,
 	const FLinearColor VisualLoggerColor)
 {
-	if (!bActive)
-	{
-		return MakeNullAction(Object, Tag);
-	}
-
-	return MakeAngleAction(Object, FMath::RadiansToDegrees(Angle), FMath::RadiansToDegrees(RelativeAngle), bActive, Tag,
+	return MakeAngleAction(Object, FMath::RadiansToDegrees(Angle), FMath::RadiansToDegrees(RelativeAngle), Tag,
 		bVisualLoggerEnabled, VisualLoggerListener, VisualLoggerAgentId, VisualLoggerAngleLocation, VisualLoggerLocation, VisualLoggerColor);
 }
 
@@ -3214,7 +2824,6 @@ FLearningAgentsActionObjectElement ULearningAgentsActions::MakeVelocityAction(
 	ULearningAgentsActionObject* Object,
 	const FVector Velocity,
 	const FTransform RelativeTransform,
-	const bool bActive,	
 	const FName Tag,
 	const bool bVisualLoggerEnabled,
 	ULearningAgentsManagerListener* VisualLoggerListener,
@@ -3223,11 +2832,6 @@ FLearningAgentsActionObjectElement ULearningAgentsActions::MakeVelocityAction(
 	const FVector VisualLoggerLocation,
 	const FLinearColor VisualLoggerColor)
 {
-	if (!bActive)
-	{
-		return MakeNullAction(Object, Tag);
-	}
-
 	const FVector LocalVelocity = RelativeTransform.InverseTransformVectorNoScale(Velocity);
 
 #if UE_LEARNING_AGENTS_ENABLE_VISUAL_LOG
@@ -3262,14 +2866,13 @@ FLearningAgentsActionObjectElement ULearningAgentsActions::MakeVelocityAction(
 		(float)LocalVelocity.X,
 		(float)LocalVelocity.Y,
 		(float)LocalVelocity.Z,
-		}, bActive, Tag);
+		}, Tag);
 }
 
 FLearningAgentsActionObjectElement ULearningAgentsActions::MakeDirectionAction(
 	ULearningAgentsActionObject* Object,
 	const FVector Direction,
 	const FTransform RelativeTransform,
-	const bool bActive,	
 	const FName Tag,
 	const bool bVisualLoggerEnabled,
 	ULearningAgentsManagerListener* VisualLoggerListener,
@@ -3279,11 +2882,6 @@ FLearningAgentsActionObjectElement ULearningAgentsActions::MakeDirectionAction(
 	const float VisualLoggerArrowLength,
 	const FLinearColor VisualLoggerColor)
 {
-	if (!bActive)
-	{
-		return MakeNullAction(Object, Tag);
-	}
-
 	const FVector LocalDirection = RelativeTransform.InverseTransformVectorNoScale(Direction).GetSafeNormal(UE_SMALL_NUMBER, FVector::ForwardVector);
 
 #if UE_LEARNING_AGENTS_ENABLE_VISUAL_LOG
@@ -3318,7 +2916,7 @@ FLearningAgentsActionObjectElement ULearningAgentsActions::MakeDirectionAction(
 		(float)LocalDirection.X,
 		(float)LocalDirection.Y,
 		(float)LocalDirection.Z,
-		}, bActive, Tag);
+		}, Tag);
 }
 
 FLearningAgentsActionModifierElement ULearningAgentsActions::MakeNullActionModifier(ULearningAgentsActionModifier* Modifier, const FName Tag)
@@ -3823,7 +3421,7 @@ FLearningAgentsActionModifierElement ULearningAgentsActions::MakeScaleActionModi
 		}, {
 		(float)LocalLogScale.X, 
 		(float)LocalLogScale.Y, 
-		(float)LocalLogScale.Z }, Tag);
+		(float)LocalLogScale.Z } , Tag);
 }
 
 FLearningAgentsActionModifierElement ULearningAgentsActions::MakeAngleActionModifier(ULearningAgentsActionModifier* Modifier, const float MaskedAngle, const bool bMask, const float RelativeAngle, const FName Tag)
@@ -3895,13 +3493,8 @@ bool ULearningAgentsActions::GetNullAction(const ULearningAgentsActionObject* Ob
 	return true;
 }
 
-bool ULearningAgentsActions::GetContinuousActionNum(int32& OutNum, const ULearningAgentsActionObject* Object, const FLearningAgentsActionObjectElement Element, const bool bActive, const FName Tag)
+bool ULearningAgentsActions::GetContinuousActionNum(int32& OutNum, const ULearningAgentsActionObject* Object, const FLearningAgentsActionObjectElement Element, const FName Tag)
 {
-	if (!bActive)
-	{
-		return GetNullAction(Object, Element, Tag);
-	}
-
 	if (!Object)
 	{
 		UE_LOG(LogLearning, Error, TEXT("GetContinuousActionNum: Object is nullptr."));
@@ -3939,7 +3532,6 @@ bool ULearningAgentsActions::GetContinuousAction(
 	TArray<float>& OutValues, 
 	const ULearningAgentsActionObject* Object, 
 	const FLearningAgentsActionObjectElement Element, 
-	const bool bActive,
 	const FName Tag,
 	const bool bVisualLoggerEnabled,
 	ULearningAgentsManagerListener* VisualLoggerListener,
@@ -3947,13 +3539,8 @@ bool ULearningAgentsActions::GetContinuousAction(
 	const FVector VisualLoggerLocation,
 	const FLinearColor VisualLoggerColor)
 {
-	if (!bActive)
-	{
-		return GetNullAction(Object, Element, Tag);
-	}
-
 	int32 OutValueNum = 0;
-	if (!GetContinuousActionNum(OutValueNum, Object, Element, bActive, Tag))
+	if (!GetContinuousActionNum(OutValueNum, Object, Element, Tag))
 	{
 		OutValues.Empty();
 		return false;
@@ -3961,7 +3548,7 @@ bool ULearningAgentsActions::GetContinuousAction(
 
 	OutValues.SetNumUninitialized(OutValueNum);
 
-	if (!GetContinuousActionToArrayView(OutValues, Object, Element, bActive, Tag, bVisualLoggerEnabled, VisualLoggerListener, VisualLoggerAgentId, VisualLoggerLocation, VisualLoggerColor))
+	if (!GetContinuousActionToArrayView(OutValues, Object, Element, Tag, bVisualLoggerEnabled, VisualLoggerListener, VisualLoggerAgentId, VisualLoggerLocation, VisualLoggerColor))
 	{
 		OutValues.Empty();
 		return false;
@@ -3974,7 +3561,6 @@ bool ULearningAgentsActions::GetContinuousActionToArrayView(
 	TArrayView<float> OutValues, 
 	const ULearningAgentsActionObject* Object, 
 	const FLearningAgentsActionObjectElement Element, 
-	const bool bActive,
 	const FName Tag,
 	const bool bVisualLoggerEnabled,
 	ULearningAgentsManagerListener* VisualLoggerListener,
@@ -3982,11 +3568,6 @@ bool ULearningAgentsActions::GetContinuousActionToArrayView(
 	const FVector VisualLoggerLocation,
 	const FLinearColor VisualLoggerColor)
 {
-	if (!bActive)
-	{
-		return GetNullAction(Object, Element, Tag);
-	}
-
 	if (!Object)
 	{
 		UE_LOG(LogLearning, Error, TEXT("GetContinuousActionToArrayView: Object is nullptr."));
@@ -4052,7 +3633,6 @@ bool ULearningAgentsActions::GetExclusiveDiscreteAction(
 	int32& OutIndex, 
 	const ULearningAgentsActionObject* Object, 
 	const FLearningAgentsActionObjectElement Element, 
-	const bool bActive,	
 	const FName Tag,
 	const bool bVisualLoggerEnabled,
 	ULearningAgentsManagerListener* VisualLoggerListener,
@@ -4060,11 +3640,6 @@ bool ULearningAgentsActions::GetExclusiveDiscreteAction(
 	const FVector VisualLoggerLocation,
 	const FLinearColor VisualLoggerColor)
 {
-	if (!bActive)
-	{
-		return GetNullAction(Object, Element, Tag);
-	}
-
 	if (!Object)
 	{
 		UE_LOG(LogLearning, Error, TEXT("GetExclusiveDiscreteAction: Object is nullptr."));
@@ -4118,7 +3693,6 @@ bool ULearningAgentsActions::GetNamedExclusiveDiscreteAction(
 	FName& OutName,
 	const ULearningAgentsActionObject* Object,
 	const FLearningAgentsActionObjectElement Element,
-	const bool bActive,
 	const FName Tag,
 	const bool bVisualLoggerEnabled,
 	ULearningAgentsManagerListener* VisualLoggerListener,
@@ -4126,11 +3700,6 @@ bool ULearningAgentsActions::GetNamedExclusiveDiscreteAction(
 	const FVector VisualLoggerLocation,
 	const FLinearColor VisualLoggerColor)
 {
-	if (!bActive)
-	{
-		return GetNullAction(Object, Element, Tag);
-	}
-
 	if (!Object)
 	{
 		UE_LOG(LogLearning, Error, TEXT("GetNamedExclusiveDiscreteAction: Object is nullptr."));
@@ -4180,13 +3749,8 @@ bool ULearningAgentsActions::GetNamedExclusiveDiscreteAction(
 	return true;
 }
 
-bool ULearningAgentsActions::GetInclusiveDiscreteActionNum(int32& OutNum, const ULearningAgentsActionObject* Object, const FLearningAgentsActionObjectElement Element, const bool bActive, const FName Tag)
+bool ULearningAgentsActions::GetInclusiveDiscreteActionNum(int32& OutNum, const ULearningAgentsActionObject* Object, const FLearningAgentsActionObjectElement Element, const FName Tag)
 {
-	if (!bActive)
-	{
-		return GetNullAction(Object, Element, Tag);
-	}
-
 	if (!Object)
 	{
 		UE_LOG(LogLearning, Error, TEXT("GetInclusiveDiscreteActionNum: Object is nullptr."));
@@ -4224,7 +3788,6 @@ bool ULearningAgentsActions::GetInclusiveDiscreteAction(
 	TArray<int32>& OutIndices, 
 	const ULearningAgentsActionObject* Object, 
 	const FLearningAgentsActionObjectElement Element, 
-	const bool bActive,
 	const FName Tag,
 	const bool bVisualLoggerEnabled,
 	ULearningAgentsManagerListener* VisualLoggerListener,
@@ -4232,13 +3795,8 @@ bool ULearningAgentsActions::GetInclusiveDiscreteAction(
 	const FVector VisualLoggerLocation,
 	const FLinearColor VisualLoggerColor)
 {
-	if (!bActive)
-	{
-		return GetNullAction(Object, Element, Tag);
-	}
-
 	int32 OutIndexNum = 0;
-	if (!GetInclusiveDiscreteActionNum(OutIndexNum, Object, Element, bActive, Tag))
+	if (!GetInclusiveDiscreteActionNum(OutIndexNum, Object, Element, Tag))
 	{
 		OutIndices.Empty();
 		return false;
@@ -4246,7 +3804,7 @@ bool ULearningAgentsActions::GetInclusiveDiscreteAction(
 
 	OutIndices.SetNumUninitialized(OutIndexNum);
 
-	if (!GetInclusiveDiscreteActionToArrayView(OutIndices, Object, Element, bActive, Tag, bVisualLoggerEnabled, VisualLoggerListener, VisualLoggerAgentId, VisualLoggerLocation, VisualLoggerColor))
+	if (!GetInclusiveDiscreteActionToArrayView(OutIndices, Object, Element, Tag, bVisualLoggerEnabled, VisualLoggerListener, VisualLoggerAgentId, VisualLoggerLocation, VisualLoggerColor))
 	{
 		OutIndices.Empty();
 		return false;
@@ -4259,7 +3817,6 @@ bool ULearningAgentsActions::GetInclusiveDiscreteActionToArrayView(
 	TArrayView<int32> OutIndices,
 	const ULearningAgentsActionObject* Object, 
 	const FLearningAgentsActionObjectElement Element, 
-	const bool bActive,
 	const FName Tag,
 	const bool bVisualLoggerEnabled,
 	ULearningAgentsManagerListener* VisualLoggerListener,
@@ -4267,11 +3824,6 @@ bool ULearningAgentsActions::GetInclusiveDiscreteActionToArrayView(
 	const FVector VisualLoggerLocation,
 	const FLinearColor VisualLoggerColor)
 {
-	if (!bActive)
-	{
-		return GetNullAction(Object, Element, Tag);
-	}
-
 	if (!Object)
 	{
 		UE_LOG(LogLearning, Error, TEXT("GetInclusiveDiscreteActionToArrayView: Object is nullptr."));
@@ -4332,13 +3884,8 @@ bool ULearningAgentsActions::GetInclusiveDiscreteActionToArrayView(
 	return true;
 }
 
-bool ULearningAgentsActions::GetNamedInclusiveDiscreteActionNum(int32& OutNum, const ULearningAgentsActionObject* Object, const FLearningAgentsActionObjectElement Element, const bool bActive, const FName Tag)
+bool ULearningAgentsActions::GetNamedInclusiveDiscreteActionNum(int32& OutNum, const ULearningAgentsActionObject* Object, const FLearningAgentsActionObjectElement Element, const FName Tag)
 {
-	if (!bActive)
-	{
-		return GetNullAction(Object, Element, Tag);
-	}
-
 	if (!Object)
 	{
 		UE_LOG(LogLearning, Error, TEXT("GetNamedInclusiveDiscreteActionNum: Object is nullptr."));
@@ -4376,7 +3923,6 @@ bool ULearningAgentsActions::GetNamedInclusiveDiscreteAction(
 	TArray<FName>& OutNames,
 	const ULearningAgentsActionObject* Object,
 	const FLearningAgentsActionObjectElement Element,
-	const bool bActive,
 	const FName Tag,
 	const bool bVisualLoggerEnabled,
 	ULearningAgentsManagerListener* VisualLoggerListener,
@@ -4384,13 +3930,8 @@ bool ULearningAgentsActions::GetNamedInclusiveDiscreteAction(
 	const FVector VisualLoggerLocation,
 	const FLinearColor VisualLoggerColor)
 {
-	if (!bActive)
-	{
-		return GetNullAction(Object, Element, Tag);
-	}
-
 	int32 OutNameNum = 0;
-	if (!GetNamedInclusiveDiscreteActionNum(OutNameNum, Object, Element, bActive, Tag))
+	if (!GetNamedInclusiveDiscreteActionNum(OutNameNum, Object, Element, Tag))
 	{
 		OutNames.Empty();
 		return false;
@@ -4398,7 +3939,7 @@ bool ULearningAgentsActions::GetNamedInclusiveDiscreteAction(
 
 	OutNames.SetNumUninitialized(OutNameNum);
 
-	if (!GetNamedInclusiveDiscreteActionToArrayView(OutNames, Object, Element, bActive, Tag, bVisualLoggerEnabled, VisualLoggerListener, VisualLoggerAgentId, VisualLoggerLocation, VisualLoggerColor))
+	if (!GetNamedInclusiveDiscreteActionToArrayView(OutNames, Object, Element, Tag, bVisualLoggerEnabled, VisualLoggerListener, VisualLoggerAgentId, VisualLoggerLocation, VisualLoggerColor))
 	{
 		OutNames.Empty();
 		return false;
@@ -4411,7 +3952,6 @@ bool ULearningAgentsActions::GetNamedInclusiveDiscreteActionToArrayView(
 	TArrayView<FName> OutNames,
 	const ULearningAgentsActionObject* Object,
 	const FLearningAgentsActionObjectElement Element,
-	const bool bActive,	
 	const FName Tag,
 	const bool bVisualLoggerEnabled,
 	ULearningAgentsManagerListener* VisualLoggerListener,
@@ -4419,11 +3959,6 @@ bool ULearningAgentsActions::GetNamedInclusiveDiscreteActionToArrayView(
 	const FVector VisualLoggerLocation,
 	const FLinearColor VisualLoggerColor)
 {
-	if (!bActive)
-	{
-		return GetNullAction(Object, Element, Tag);
-	}
-
 	if (!Object)
 	{
 		UE_LOG(LogLearning, Error, TEXT("GetNamedInclusiveDiscreteActionToArrayView: Object is nullptr."));
@@ -4485,13 +4020,8 @@ bool ULearningAgentsActions::GetNamedInclusiveDiscreteActionToArrayView(
 }
 
 
-bool ULearningAgentsActions::GetStructActionNum(int32& OutNum, const ULearningAgentsActionObject* Object, const FLearningAgentsActionObjectElement Element, const bool bActive, const FName Tag)
+bool ULearningAgentsActions::GetStructActionNum(int32& OutNum, const ULearningAgentsActionObject* Object, const FLearningAgentsActionObjectElement Element, const FName Tag)
 {
-	if (!bActive)
-	{
-		return GetNullAction(Object, Element, Tag);
-	}
-
 	if (!Object)
 	{
 		UE_LOG(LogLearning, Error, TEXT("GetStructActionNum: Object is nullptr."));
@@ -4527,15 +4057,10 @@ bool ULearningAgentsActions::GetStructActionNum(int32& OutNum, const ULearningAg
 	return true;
 }
 
-bool ULearningAgentsActions::GetStructAction(TMap<FName, FLearningAgentsActionObjectElement>& OutElements, const ULearningAgentsActionObject* Object, const FLearningAgentsActionObjectElement Element, const bool bActive, const FName Tag)
+bool ULearningAgentsActions::GetStructAction(TMap<FName, FLearningAgentsActionObjectElement>& OutElements, const ULearningAgentsActionObject* Object, const FLearningAgentsActionObjectElement Element, const FName Tag)
 {
-	if (!bActive)
-	{
-		return GetNullAction(Object, Element, Tag);
-	}
-
 	int32 OutElementNum = 0;
-	if (!GetStructActionNum(OutElementNum, Object, Element, bActive, Tag))
+	if (!GetStructActionNum(OutElementNum, Object, Element, Tag))
 	{
 		OutElements.Empty();
 		return false;
@@ -4546,7 +4071,7 @@ bool ULearningAgentsActions::GetStructAction(TMap<FName, FLearningAgentsActionOb
 	SubElementNames.SetNumUninitialized(OutElementNum);
 	SubElements.SetNumUninitialized(OutElementNum);
 
-	if (!GetStructActionToArrayViews(SubElementNames, SubElements, Object, Element, bActive, Tag))
+	if (!GetStructActionToArrayViews(SubElementNames, SubElements, Object, Element, Tag))
 	{
 		OutElements.Empty();
 		return false;
@@ -4561,13 +4086,8 @@ bool ULearningAgentsActions::GetStructAction(TMap<FName, FLearningAgentsActionOb
 	return true;
 }
 
-bool ULearningAgentsActions::GetStructActionElement(FLearningAgentsActionObjectElement& OutElement, const ULearningAgentsActionObject* Object, const FLearningAgentsActionObjectElement Element, const FName ElementName, const bool bActive, const FName Tag)
+bool ULearningAgentsActions::GetStructActionElement(FLearningAgentsActionObjectElement& OutElement, const ULearningAgentsActionObject* Object, const FLearningAgentsActionObjectElement Element, const FName ElementName, const FName Tag)
 {
-	if (!bActive)
-	{
-		return GetNullAction(Object, Element, Tag);
-	}
-
 	if (!Object)
 	{
 		UE_LOG(LogLearning, Error, TEXT("GetStructActionElement: Object is nullptr."));
@@ -4612,15 +4132,10 @@ bool ULearningAgentsActions::GetStructActionElement(FLearningAgentsActionObjectE
 	return true;
 }
 
-bool ULearningAgentsActions::GetStructActionToArrays(TArray<FName>& OutElementNames, TArray<FLearningAgentsActionObjectElement>& OutElements, const ULearningAgentsActionObject* Object, const FLearningAgentsActionObjectElement Element, const bool bActive, const FName Tag)
+bool ULearningAgentsActions::GetStructActionToArrays(TArray<FName>& OutElementNames, TArray<FLearningAgentsActionObjectElement>& OutElements, const ULearningAgentsActionObject* Object, const FLearningAgentsActionObjectElement Element, const FName Tag)
 {
-	if (!bActive)
-	{
-		return GetNullAction(Object, Element, Tag);
-	}
-
 	int32 OutElementNum = 0;
-	if (!GetStructActionNum(OutElementNum, Object, Element, bActive, Tag))
+	if (!GetStructActionNum(OutElementNum, Object, Element, Tag))
 	{
 		OutElementNames.Empty();
 		OutElements.Empty();
@@ -4630,7 +4145,7 @@ bool ULearningAgentsActions::GetStructActionToArrays(TArray<FName>& OutElementNa
 	OutElementNames.SetNumUninitialized(OutElementNum);
 	OutElements.SetNumUninitialized(OutElementNum);
 
-	if (!GetStructActionToArrayViews(OutElementNames, OutElements, Object, Element, bActive, Tag))
+	if (!GetStructActionToArrayViews(OutElementNames, OutElements, Object, Element, Tag))
 	{
 		OutElementNames.Empty();
 		OutElements.Empty();
@@ -4640,13 +4155,8 @@ bool ULearningAgentsActions::GetStructActionToArrays(TArray<FName>& OutElementNa
 	return true;
 }
 
-bool ULearningAgentsActions::GetStructActionToArrayViews(TArrayView<FName> OutElementNames, TArrayView<FLearningAgentsActionObjectElement> OutElements, const ULearningAgentsActionObject* Object, const FLearningAgentsActionObjectElement Element, const bool bActive, const FName Tag)
+bool ULearningAgentsActions::GetStructActionToArrayViews(TArrayView<FName> OutElementNames, TArrayView<FLearningAgentsActionObjectElement> OutElements, const ULearningAgentsActionObject* Object, const FLearningAgentsActionObjectElement Element, const FName Tag)
 {
-	if (!bActive)
-	{
-		return GetNullAction(Object, Element, Tag);
-	}
-
 	if (!Object)
 	{
 		UE_LOG(LogLearning, Error, TEXT("GetStructActionToArrayViews: Object is nullptr."));
@@ -4713,13 +4223,8 @@ bool ULearningAgentsActions::GetStructActionToArrayViews(TArrayView<FName> OutEl
 	return true;
 }
 
-bool ULearningAgentsActions::GetExclusiveUnionAction(FName& OutElementName, FLearningAgentsActionObjectElement& OutElement, const ULearningAgentsActionObject* Object, const FLearningAgentsActionObjectElement Element, const bool bActive, const FName Tag)
+bool ULearningAgentsActions::GetExclusiveUnionAction(FName& OutElementName, FLearningAgentsActionObjectElement& OutElement, const ULearningAgentsActionObject* Object, const FLearningAgentsActionObjectElement Element, const FName Tag)
 {
-	if (!bActive)
-	{
-		return GetNullAction(Object, Element, Tag);
-	}
-
 	if (!Object)
 	{
 		UE_LOG(LogLearning, Error, TEXT("GetExclusiveUnionAction: Object is nullptr."));
@@ -4759,13 +4264,8 @@ bool ULearningAgentsActions::GetExclusiveUnionAction(FName& OutElementName, FLea
 	return true;
 }
 
-bool ULearningAgentsActions::GetInclusiveUnionActionNum(int32& OutNum, const ULearningAgentsActionObject* Object, const FLearningAgentsActionObjectElement Element, const bool bActive, const FName Tag)
+bool ULearningAgentsActions::GetInclusiveUnionActionNum(int32& OutNum, const ULearningAgentsActionObject* Object, const FLearningAgentsActionObjectElement Element, const FName Tag)
 {
-	if (!bActive)
-	{
-		return GetNullAction(Object, Element, Tag);
-	}
-
 	if (!Object)
 	{
 		UE_LOG(LogLearning, Error, TEXT("GetInclusiveUnionActionNum: Object is nullptr."));
@@ -4801,15 +4301,10 @@ bool ULearningAgentsActions::GetInclusiveUnionActionNum(int32& OutNum, const ULe
 	return true;
 }
 
-bool ULearningAgentsActions::GetInclusiveUnionAction(TMap<FName, FLearningAgentsActionObjectElement>& OutElements, const ULearningAgentsActionObject* Object, const FLearningAgentsActionObjectElement Element, const bool bActive, const FName Tag)
+bool ULearningAgentsActions::GetInclusiveUnionAction(TMap<FName, FLearningAgentsActionObjectElement>& OutElements, const ULearningAgentsActionObject* Object, const FLearningAgentsActionObjectElement Element, const FName Tag)
 {
-	if (!bActive)
-	{
-		return GetNullAction(Object, Element, Tag);
-	}
-
 	int32 OutElementNum = 0;
-	if (!GetInclusiveUnionActionNum(OutElementNum, Object, Element, bActive, Tag))
+	if (!GetInclusiveUnionActionNum(OutElementNum, Object, Element, Tag))
 	{
 		OutElements.Empty();
 		return false;
@@ -4820,7 +4315,7 @@ bool ULearningAgentsActions::GetInclusiveUnionAction(TMap<FName, FLearningAgents
 	SubElementNames.SetNumUninitialized(OutElementNum);
 	SubElements.SetNumUninitialized(OutElementNum);
 
-	if (!GetInclusiveUnionActionToArrayViews(SubElementNames, SubElements, Object, Element, bActive, Tag))
+	if (!GetInclusiveUnionActionToArrayViews(SubElementNames, SubElements, Object, Element, Tag))
 	{
 		OutElements.Empty();
 		return false;
@@ -4835,15 +4330,10 @@ bool ULearningAgentsActions::GetInclusiveUnionAction(TMap<FName, FLearningAgents
 	return true;
 }
 
-bool ULearningAgentsActions::GetInclusiveUnionActionToArrays(TArray<FName>& OutElementNames, TArray<FLearningAgentsActionObjectElement>& OutElements, const ULearningAgentsActionObject* Object, const FLearningAgentsActionObjectElement Element, const bool bActive, const FName Tag)
+bool ULearningAgentsActions::GetInclusiveUnionActionToArrays(TArray<FName>& OutElementNames, TArray<FLearningAgentsActionObjectElement>& OutElements, const ULearningAgentsActionObject* Object, const FLearningAgentsActionObjectElement Element, const FName Tag)
 {
-	if (!bActive)
-	{
-		return GetNullAction(Object, Element, Tag);
-	}
-
 	int32 OutElementNum = 0;
-	if (!GetInclusiveUnionActionNum(OutElementNum, Object, Element, bActive, Tag))
+	if (!GetInclusiveUnionActionNum(OutElementNum, Object, Element, Tag))
 	{
 		OutElementNames.Empty();
 		OutElements.Empty();
@@ -4853,7 +4343,7 @@ bool ULearningAgentsActions::GetInclusiveUnionActionToArrays(TArray<FName>& OutE
 	OutElementNames.SetNumUninitialized(OutElementNum);
 	OutElements.SetNumUninitialized(OutElementNum);
 
-	if (!GetInclusiveUnionActionToArrayViews(OutElementNames, OutElements, Object, Element, bActive, Tag))
+	if (!GetInclusiveUnionActionToArrayViews(OutElementNames, OutElements, Object, Element, Tag))
 	{
 		OutElementNames.Empty();
 		OutElements.Empty();
@@ -4863,13 +4353,8 @@ bool ULearningAgentsActions::GetInclusiveUnionActionToArrays(TArray<FName>& OutE
 	return true;
 }
 
-bool ULearningAgentsActions::GetInclusiveUnionActionToArrayViews(TArrayView<FName> OutElementNames, TArrayView<FLearningAgentsActionObjectElement> OutElements, const ULearningAgentsActionObject* Object, const FLearningAgentsActionObjectElement Element, const bool bActive, const FName Tag)
+bool ULearningAgentsActions::GetInclusiveUnionActionToArrayViews(TArrayView<FName> OutElementNames, TArrayView<FLearningAgentsActionObjectElement> OutElements, const ULearningAgentsActionObject* Object, const FLearningAgentsActionObjectElement Element, const FName Tag)
 {
-	if (!bActive)
-	{
-		return GetNullAction(Object, Element, Tag);
-	}
-
 	if (!Object)
 	{
 		UE_LOG(LogLearning, Error, TEXT("GetInclusiveUnionActionToArrayViews: Object is nullptr."));
@@ -4931,13 +4416,8 @@ bool ULearningAgentsActions::GetInclusiveUnionActionToArrayViews(TArrayView<FNam
 	return true;
 }
 
-bool ULearningAgentsActions::GetStaticArrayActionNum(int32& OutNum, const ULearningAgentsActionObject* Object, const FLearningAgentsActionObjectElement Element, const bool bActive, const FName Tag)
+bool ULearningAgentsActions::GetStaticArrayActionNum(int32& OutNum, const ULearningAgentsActionObject* Object, const FLearningAgentsActionObjectElement Element, const FName Tag)
 {
-	if (!bActive)
-	{
-		return GetNullAction(Object, Element, Tag);
-	}
-
 	if (!Object)
 	{
 		UE_LOG(LogLearning, Error, TEXT("GetStaticArrayActionNum: Object is nullptr."));
@@ -4971,15 +4451,10 @@ bool ULearningAgentsActions::GetStaticArrayActionNum(int32& OutNum, const ULearn
 	return true;
 }
 
-bool ULearningAgentsActions::GetStaticArrayAction(TArray<FLearningAgentsActionObjectElement>& OutElements, const ULearningAgentsActionObject* Object, const FLearningAgentsActionObjectElement Element, const bool bActive, const FName Tag)
+bool ULearningAgentsActions::GetStaticArrayAction(TArray<FLearningAgentsActionObjectElement>& OutElements, const ULearningAgentsActionObject* Object, const FLearningAgentsActionObjectElement Element, const FName Tag)
 {
-	if (!bActive)
-	{
-		return GetNullAction(Object, Element, Tag);
-	}
-
 	int32 OutElementNum = 0;
-	if (!GetStaticArrayActionNum(OutElementNum, Object, Element, bActive, Tag))
+	if (!GetStaticArrayActionNum(OutElementNum, Object, Element, Tag))
 	{
 		OutElements.Empty();
 		return false;
@@ -4987,7 +4462,7 @@ bool ULearningAgentsActions::GetStaticArrayAction(TArray<FLearningAgentsActionOb
 
 	OutElements.SetNumUninitialized(OutElementNum);
 
-	if (!GetStaticArrayActionToArrayView(OutElements, Object, Element, bActive, Tag))
+	if (!GetStaticArrayActionToArrayView(OutElements, Object, Element, Tag))
 	{
 		OutElements.Empty();
 		return false;
@@ -4996,13 +4471,8 @@ bool ULearningAgentsActions::GetStaticArrayAction(TArray<FLearningAgentsActionOb
 	return true;
 }
 
-bool ULearningAgentsActions::GetStaticArrayActionToArrayView(TArrayView<FLearningAgentsActionObjectElement> OutElements, const ULearningAgentsActionObject* Object, const FLearningAgentsActionObjectElement Element, const bool bActive, const FName Tag)
+bool ULearningAgentsActions::GetStaticArrayActionToArrayView(TArrayView<FLearningAgentsActionObjectElement> OutElements, const ULearningAgentsActionObject* Object, const FLearningAgentsActionObjectElement Element, const FName Tag)
 {
-	if (!bActive)
-	{
-		return GetNullAction(Object, Element, Tag);
-	}
-
 	if (!Object)
 	{
 		UE_LOG(LogLearning, Error, TEXT("GetStaticArrayActionToArrayView: Object is nullptr."));
@@ -5063,16 +4533,11 @@ bool ULearningAgentsActions::GetStaticArrayActionToArrayView(TArrayView<FLearnin
 	return true;
 }
 
-bool ULearningAgentsActions::GetPairAction(FLearningAgentsActionObjectElement& OutKey, FLearningAgentsActionObjectElement& OutValue, const ULearningAgentsActionObject* Object, const FLearningAgentsActionObjectElement Element, const bool bActive, const FName Tag)
+bool ULearningAgentsActions::GetPairAction(FLearningAgentsActionObjectElement& OutKey, FLearningAgentsActionObjectElement& OutValue, const ULearningAgentsActionObject* Object, const FLearningAgentsActionObjectElement Element, const FName Tag)
 {
-	if (!bActive)
-	{
-		return GetNullAction(Object, Element, Tag);
-	}
-
 	TStaticArray<FName, 2> OutElementNames;
 	TStaticArray<FLearningAgentsActionObjectElement, 2> OutElements;
-	if (!GetStructActionToArrayViews(OutElementNames, OutElements, Object, Element, bActive, Tag))
+	if (!GetStructActionToArrayViews(OutElementNames, OutElements, Object, Element, Tag))
 	{
 		OutKey = FLearningAgentsActionObjectElement();
 		OutValue = FLearningAgentsActionObjectElement();
@@ -5089,7 +4554,6 @@ bool ULearningAgentsActions::GetEnumAction(
 	const ULearningAgentsActionObject* Object, 
 	const FLearningAgentsActionObjectElement Element, 
 	const UEnum* Enum, 
-	const bool bActive,
 	const FName Tag,
 	const bool bVisualLoggerEnabled,
 	ULearningAgentsManagerListener* VisualLoggerListener,
@@ -5097,11 +4561,6 @@ bool ULearningAgentsActions::GetEnumAction(
 	const FVector VisualLoggerLocation,
 	const FLinearColor VisualLoggerColor)
 {
-	if (!bActive)
-	{
-		return GetNullAction(Object, Element, Tag);
-	}
-
 	if (!Enum)
 	{
 		UE_LOG(LogLearning, Error, TEXT("GetEnumAction: Enum is nullptr."));
@@ -5110,7 +4569,7 @@ bool ULearningAgentsActions::GetEnumAction(
 	}
 
 	int32 OutIndex = 0;
-	if (!GetExclusiveDiscreteAction(OutIndex, Object, Element, bActive, Tag))
+	if (!GetExclusiveDiscreteAction(OutIndex, Object, Element, Tag))
 	{
 		OutEnumValue = 0;
 		return false;
@@ -5160,7 +4619,6 @@ bool ULearningAgentsActions::GetBitmaskAction(
 	const ULearningAgentsActionObject* Object, 
 	const FLearningAgentsActionObjectElement Element, 
 	const UEnum* Enum, 
-	const bool bActive,
 	const FName Tag,
 	const bool bVisualLoggerEnabled,
 	ULearningAgentsManagerListener* VisualLoggerListener,
@@ -5168,11 +4626,6 @@ bool ULearningAgentsActions::GetBitmaskAction(
 	const FVector VisualLoggerLocation,
 	const FLinearColor VisualLoggerColor)
 {
-	if (!bActive)
-	{
-		return GetNullAction(Object, Element, Tag);
-	}
-
 	if (!Enum)
 	{
 		UE_LOG(LogLearning, Error, TEXT("GetBitmaskAction: Enum is nullptr."));
@@ -5188,7 +4641,7 @@ bool ULearningAgentsActions::GetBitmaskAction(
 	}
 
 	int32 EnumValueNum;
-	if (!GetInclusiveDiscreteActionNum(EnumValueNum, Object, Element, bActive, Tag))
+	if (!GetInclusiveDiscreteActionNum(EnumValueNum, Object, Element, Tag))
 	{
 		OutBitmaskValue = 0;
 		return false;
@@ -5203,7 +4656,7 @@ bool ULearningAgentsActions::GetBitmaskAction(
 
 	TArray<int32, TInlineAllocator<32>> OutIndices;
 	OutIndices.SetNumUninitialized(EnumValueNum);
-	if (!GetInclusiveDiscreteActionToArrayView(OutIndices, Object, Element, bActive, Tag))
+	if (!GetInclusiveDiscreteActionToArrayView(OutIndices, Object, Element, Tag))
 	{
 		OutBitmaskValue = 0;
 		return false;
@@ -5251,15 +4704,10 @@ bool ULearningAgentsActions::GetBitmaskAction(
 	return true;
 }
 
-bool ULearningAgentsActions::GetOptionalAction(ELearningAgentsOptionalAction& OutOption, FLearningAgentsActionObjectElement& OutElement, const ULearningAgentsActionObject* Object, const FLearningAgentsActionObjectElement Element, const bool bActive, const FName Tag)
+bool ULearningAgentsActions::GetOptionalAction(ELearningAgentsOptionalAction& OutOption, FLearningAgentsActionObjectElement& OutElement, const ULearningAgentsActionObject* Object, const FLearningAgentsActionObjectElement Element, const FName Tag)
 {
-	if (!bActive)
-	{
-		return GetNullAction(Object, Element, Tag);
-	}
-
 	FName OutName = NAME_None;
-	if (!GetExclusiveUnionAction(OutName, OutElement, Object, Element, bActive, Tag))
+	if (!GetExclusiveUnionAction(OutName, OutElement, Object, Element, Tag))
 	{
 		OutOption = ELearningAgentsOptionalAction::Null;
 		return false;
@@ -5269,15 +4717,10 @@ bool ULearningAgentsActions::GetOptionalAction(ELearningAgentsOptionalAction& Ou
 	return true;
 }
 
-bool ULearningAgentsActions::GetEitherAction(ELearningAgentsEitherAction& OutEither, FLearningAgentsActionObjectElement& OutElement, const ULearningAgentsActionObject* Object, const FLearningAgentsActionObjectElement Element, const bool bActive, const FName Tag)
+bool ULearningAgentsActions::GetEitherAction(ELearningAgentsEitherAction& OutEither, FLearningAgentsActionObjectElement& OutElement, const ULearningAgentsActionObject* Object, const FLearningAgentsActionObjectElement Element, const FName Tag)
 {
-	if (!bActive)
-	{
-		return GetNullAction(Object, Element, Tag);
-	}
-
 	FName OutName = NAME_None;
-	if (!GetExclusiveUnionAction(OutName, OutElement, Object, Element, bActive, Tag))
+	if (!GetExclusiveUnionAction(OutName, OutElement, Object, Element, Tag))
 	{
 		OutEither = ELearningAgentsEitherAction::A;
 		return false;
@@ -5287,13 +4730,8 @@ bool ULearningAgentsActions::GetEitherAction(ELearningAgentsEitherAction& OutEit
 	return true;
 }
 
-bool ULearningAgentsActions::GetEncodingAction(FLearningAgentsActionObjectElement& OutElement, const ULearningAgentsActionObject* Object, const FLearningAgentsActionObjectElement Element, const bool bActive, const FName Tag)
+bool ULearningAgentsActions::GetEncodingAction(FLearningAgentsActionObjectElement& OutElement, const ULearningAgentsActionObject* Object, const FLearningAgentsActionObjectElement Element, const FName Tag)
 {
-	if (!bActive)
-	{
-		return GetNullAction(Object, Element, Tag);
-	}
-
 	if (!Object)
 	{
 		UE_LOG(LogLearning, Error, TEXT("GetEncodingAction: Object is nullptr."));
@@ -5332,7 +4770,6 @@ bool ULearningAgentsActions::GetBoolAction(
 	bool& bOutValue, 
 	const ULearningAgentsActionObject* Object, 
 	const FLearningAgentsActionObjectElement Element, 
-	const bool bActive,	
 	const FName Tag,
 	const bool bVisualLoggerEnabled,
 	ULearningAgentsManagerListener* VisualLoggerListener,
@@ -5340,13 +4777,8 @@ bool ULearningAgentsActions::GetBoolAction(
 	const FVector VisualLoggerLocation,
 	const FLinearColor VisualLoggerColor)
 {
-	if (!bActive)
-	{
-		return GetNullAction(Object, Element, Tag);
-	}
-
 	int32 OutIndex = 0;
-	if (!GetExclusiveDiscreteAction(OutIndex, Object, Element, bActive, Tag))
+	if (!GetExclusiveDiscreteAction(OutIndex, Object, Element, Tag))
 	{
 		bOutValue = false;
 		return false;
@@ -5376,7 +4808,6 @@ bool ULearningAgentsActions::GetFloatAction(
 	float& OutValue, 
 	const ULearningAgentsActionObject* Object, 
 	const FLearningAgentsActionObjectElement Element, 
-	const bool bActive, 
 	const FName Tag,
 	const bool bVisualLoggerEnabled,
 	ULearningAgentsManagerListener* VisualLoggerListener,
@@ -5384,12 +4815,7 @@ bool ULearningAgentsActions::GetFloatAction(
 	const FVector VisualLoggerLocation,
 	const FLinearColor VisualLoggerColor)
 {
-	if (!bActive)
-	{
-		return GetNullAction(Object, Element, Tag);
-	}
-
-	if (!GetContinuousActionToArrayView(MakeArrayView(&OutValue, 1), Object, Element, bActive, Tag))
+	if (!GetContinuousActionToArrayView(MakeArrayView(&OutValue, 1), Object, Element, Tag))
 	{
 		OutValue = 0.0f;
 		return false;
@@ -5418,7 +4844,6 @@ bool ULearningAgentsActions::GetLocationAction(
 	const ULearningAgentsActionObject* Object, 
 	const FLearningAgentsActionObjectElement Element, 
 	const FTransform RelativeTransform, 
-	const bool bActive,
 	const FName Tag,
 	const bool bVisualLoggerEnabled,
 	ULearningAgentsManagerListener* VisualLoggerListener,
@@ -5426,13 +4851,8 @@ bool ULearningAgentsActions::GetLocationAction(
 	const FVector VisualLoggerLocation,
 	const FLinearColor VisualLoggerColor)
 {
-	if (!bActive)
-	{
-		return GetNullAction(Object, Element, Tag);
-	}
-
 	TStaticArray<float, 3> OutValues;
-	if (!GetContinuousActionToArrayView(OutValues, Object, Element, bActive, Tag))
+	if (!GetContinuousActionToArrayView(OutValues, Object, Element, Tag))
 	{
 		OutLocation = FVector::ZeroVector;
 		return false;
@@ -5483,7 +4903,6 @@ bool ULearningAgentsActions::GetRotationAction(
 	const ULearningAgentsActionObject* Object, 
 	const FLearningAgentsActionObjectElement Element, 
 	const FRotator RelativeRotation, 
-	const bool bActive,
 	const FName Tag,
 	const bool bVisualLoggerEnabled,
 	ULearningAgentsManagerListener* VisualLoggerListener,
@@ -5492,13 +4911,8 @@ bool ULearningAgentsActions::GetRotationAction(
 	const FVector VisualLoggerLocation,
 	const FLinearColor VisualLoggerColor)
 {
-	if (!bActive)
-	{
-		return GetNullAction(Object, Element, Tag);
-	}
-
 	FQuat OutRotationQuat;
-	if (!GetRotationActionAsQuat(OutRotationQuat, Object, Element, FQuat::MakeFromRotator(RelativeRotation), bActive, Tag, bVisualLoggerEnabled, VisualLoggerListener, VisualLoggerAgentId, VisualLoggerRotationLocation, VisualLoggerLocation, VisualLoggerColor))
+	if (!GetRotationActionAsQuat(OutRotationQuat, Object, Element, FQuat::MakeFromRotator(RelativeRotation), Tag, bVisualLoggerEnabled, VisualLoggerListener, VisualLoggerAgentId, VisualLoggerRotationLocation, VisualLoggerLocation, VisualLoggerColor))
 	{
 		OutRotation = FRotator::ZeroRotator;
 		return false;
@@ -5513,7 +4927,6 @@ bool ULearningAgentsActions::GetRotationActionAsQuat(
 	const ULearningAgentsActionObject* Object, 
 	const FLearningAgentsActionObjectElement Element, 
 	const FQuat RelativeRotation, 
-	const bool bActive,
 	const FName Tag,
 	const bool bVisualLoggerEnabled,
 	ULearningAgentsManagerListener* VisualLoggerListener,
@@ -5522,13 +4935,8 @@ bool ULearningAgentsActions::GetRotationActionAsQuat(
 	const FVector VisualLoggerLocation,
 	const FLinearColor VisualLoggerColor)
 {
-	if (!bActive)
-	{
-		return GetNullAction(Object, Element, Tag);
-	}
-
 	TStaticArray<float, 3> OutValues;
-	if (!GetContinuousActionToArrayView(OutValues, Object, Element, bActive, Tag))
+	if (!GetContinuousActionToArrayView(OutValues, Object, Element, Tag))
 	{
 		OutRotation = FQuat::Identity;
 		return false;
@@ -5569,7 +4977,6 @@ bool ULearningAgentsActions::GetScaleAction(
 	const ULearningAgentsActionObject* Object, 
 	const FLearningAgentsActionObjectElement Element, 
 	const FVector RelativeScale, 
-	const bool bActive,
 	const FName Tag,
 	const bool bVisualLoggerEnabled,
 	ULearningAgentsManagerListener* VisualLoggerListener,
@@ -5577,13 +4984,8 @@ bool ULearningAgentsActions::GetScaleAction(
 	const FVector VisualLoggerLocation,
 	const FLinearColor VisualLoggerColor)
 {
-	if (!bActive)
-	{
-		return GetNullAction(Object, Element, Tag);
-	}
-
 	TStaticArray<float, 3> OutValues;
-	if (!GetContinuousActionToArrayView(OutValues, Object, Element, bActive, Tag))
+	if (!GetContinuousActionToArrayView(OutValues, Object, Element, Tag))
 	{
 		OutScale = FVector::OneVector;
 		return false;
@@ -5616,7 +5018,6 @@ bool ULearningAgentsActions::GetTransformAction(
 	const ULearningAgentsActionObject* Object, 
 	const FLearningAgentsActionObjectElement Element, 
 	const FTransform RelativeTransform, 
-	const bool bActive,
 	const FName Tag,
 	const bool bVisualLoggerEnabled,
 	ULearningAgentsManagerListener* VisualLoggerListener,
@@ -5624,14 +5025,9 @@ bool ULearningAgentsActions::GetTransformAction(
 	const FVector VisualLoggerLocation,
 	const FLinearColor VisualLoggerColor)
 {
-	if (!bActive)
-	{
-		return GetNullAction(Object, Element, Tag);
-	}
-
 	TStaticArray<FName, 3> OutElementNames;
 	TStaticArray<FLearningAgentsActionObjectElement, 3> OutElements;
-	if (!GetStructActionToArrayViews(OutElementNames, OutElements, Object, Element, bActive, Tag))
+	if (!GetStructActionToArrayViews(OutElementNames, OutElements, Object, Element, Tag))
 	{
 		OutTransform = FTransform::Identity;
 		return false;
@@ -5691,7 +5087,6 @@ bool ULearningAgentsActions::GetAngleAction(
 	const ULearningAgentsActionObject* Object, 
 	const FLearningAgentsActionObjectElement Element, 
 	const float RelativeAngle, 
-	const bool bActive,
 	const FName Tag,
 	const bool bVisualLoggerEnabled,
 	ULearningAgentsManagerListener* VisualLoggerListener,
@@ -5700,13 +5095,8 @@ bool ULearningAgentsActions::GetAngleAction(
 	const FVector VisualLoggerLocation,
 	const FLinearColor VisualLoggerColor)
 {
-	if (!bActive)
-	{
-		return GetNullAction(Object, Element, Tag);
-	}
-
 	float LocalAngle = 0.0f;
-	if (!GetContinuousActionToArrayView(MakeArrayView(&LocalAngle, 1), Object, Element, bActive, Tag))
+	if (!GetContinuousActionToArrayView(MakeArrayView(&LocalAngle, 1), Object, Element, Tag))
 	{
 		OutAngle = 0.0f;
 		return false;
@@ -5746,7 +5136,6 @@ bool ULearningAgentsActions::GetAngleActionRadians(
 	const ULearningAgentsActionObject* Object, 
 	const FLearningAgentsActionObjectElement Element, 
 	const float RelativeAngle, 
-	const bool bActive,	
 	const FName Tag,
 	const bool bVisualLoggerEnabled,
 	ULearningAgentsManagerListener* VisualLoggerListener,
@@ -5755,12 +5144,7 @@ bool ULearningAgentsActions::GetAngleActionRadians(
 	const FVector VisualLoggerLocation,
 	const FLinearColor VisualLoggerColor)
 {
-	if (!bActive)
-	{
-		return GetNullAction(Object, Element, Tag);
-	}
-
-	if (!GetAngleAction(OutAngle, Object, Element, FMath::RadiansToDegrees(RelativeAngle), bActive, Tag, bVisualLoggerEnabled, VisualLoggerListener, VisualLoggerAgentId, VisualLoggerAngleLocation, VisualLoggerLocation, VisualLoggerColor))
+	if (!GetAngleAction(OutAngle, Object, Element, FMath::RadiansToDegrees(RelativeAngle), Tag, bVisualLoggerEnabled, VisualLoggerListener, VisualLoggerAgentId, VisualLoggerAngleLocation, VisualLoggerLocation, VisualLoggerColor))
 	{
 		OutAngle = 0.0f;
 		return false;
@@ -5775,7 +5159,6 @@ bool ULearningAgentsActions::GetVelocityAction(
 	const ULearningAgentsActionObject* Object, 
 	const FLearningAgentsActionObjectElement Element, 
 	const FTransform RelativeTransform, 
-	const bool bActive,
 	const FName Tag,
 	const bool bVisualLoggerEnabled,
 	ULearningAgentsManagerListener* VisualLoggerListener,
@@ -5784,13 +5167,8 @@ bool ULearningAgentsActions::GetVelocityAction(
 	const FVector VisualLoggerLocation,
 	const FLinearColor VisualLoggerColor)
 {
-	if (!bActive)
-	{
-		return GetNullAction(Object, Element, Tag);
-	}
-
 	TStaticArray<float, 3> OutValues;
-	if (!GetContinuousActionToArrayView(OutValues, Object, Element, bActive, Tag))
+	if (!GetContinuousActionToArrayView(OutValues, Object, Element, Tag))
 	{
 		OutVelocity = FVector::OneVector;
 		return false;
@@ -5835,7 +5213,6 @@ bool ULearningAgentsActions::GetDirectionAction(
 	const ULearningAgentsActionObject* Object, 
 	const FLearningAgentsActionObjectElement Element, 
 	const FTransform RelativeTransform, 
-	const bool bActive,	
 	const FName Tag,
 	const bool bVisualLoggerEnabled,
 	ULearningAgentsManagerListener* VisualLoggerListener,
@@ -5845,13 +5222,8 @@ bool ULearningAgentsActions::GetDirectionAction(
 	const float VisualLoggerArrowLength,
 	const FLinearColor VisualLoggerColor)
 {
-	if (!bActive)
-	{
-		return GetNullAction(Object, Element, Tag);
-	}
-
 	TStaticArray<float, 3> OutValues;
-	if (!GetContinuousActionToArrayView(OutValues, Object, Element, bActive, Tag))
+	if (!GetContinuousActionToArrayView(OutValues, Object, Element, Tag))
 	{
 		OutDirection = FVector::ForwardVector;
 		return false;
