[LearningAgents] Flow Matching and Normalizations

+ network embedded normalizations
+ flow matching training objective

- update anim control as a result of new normalization code
- update LACombat as a result of new normalization code

#Synced-CL 51518158

diff --git a/Engine/Plugins/Experimental/Animation/AnimGen/Source/AnimGen/Private/AnimGenControl.cpp b/Engine/Plugins/Experimental/Animation/AnimGen/Source/AnimGen/Private/AnimGenControl.cpp
--- a/Engine/Plugins/Experimental/Animation/AnimGen/Source/AnimGen/Private/AnimGenControl.cpp
+++ b/Engine/Plugins/Experimental/Animation/AnimGen/Source/AnimGen/Private/AnimGenControl.cpp
@@ -735,9 +735,13 @@ FAnimGenControlSchemaElement UAnimGenControls::SpecifyContinuousControl(FAnimGen
 		UE_ANIMGEN_LOG_ONCE(LogAnimGen, Warning, "SpecifyContinuousControl: Specifying zero-sized Continuous Control.");
 	}
 
+	TArray<float, TInlineAllocator<128>> Scales, Offsets;
+	Scales.Init(Scale, Size);
+	Offsets.Init(0.0f, Size);
+
 	{
 		FWriteScopeLock SchemaLock(Schema.ObservationSchema->Lock);
-		return { Schema.ObservationSchema->Schema.CreateContinuous({ Size, Scale }, Tag) };
+		return { Schema.ObservationSchema->Schema.CreateContinuous({ Size, Scales, Offsets, UE::Learning::Observation::ENormalization::Manual }, Tag) };
 	}
 }
 
diff --git a/Engine/Plugins/Experimental/LearningAgents/Source/LearningAgents/Private/LearningAgentsActions.cpp b/Engine/Plugins/Experimental/LearningAgents/Source/LearningAgents/Private/LearningAgentsActions.cpp
--- a/Engine/Plugins/Experimental/LearningAgents/Source/LearningAgents/Private/LearningAgentsActions.cpp
+++ b/Engine/Plugins/Experimental/LearningAgents/Source/LearningAgents/Private/LearningAgentsActions.cpp
@@ -32,7 +32,7 @@ uint32 GetTypeHash(const FLearningAgentsActionSchemaElement& Element)
 	return (uint32)Element.SchemaElement.Index;
 }
 
-uint32 GetTypeHash(const FLearningAgentsActionObjectElement & Element)
+uint32 GetTypeHash(const FLearningAgentsActionObjectElement& Element)
 {
 	return (uint32)Element.ObjectElement.Index;
 }
@@ -387,10 +387,10 @@ namespace UE::Learning::Agents::Action::Private
 			const Learning::Action::FObjectEncodingParameters ObjectParameters = Object.GetEncoding(ObjectElement);
 
 			return ValidateActionObjectMatchesSchema(
-					Schema,
-					SchemaParameters.Element,
-					Object,
-					ObjectParameters.Element);
+				Schema,
+				SchemaParameters.Element,
+				Object,
+				ObjectParameters.Element);
 		}
 
 		default:
@@ -922,6 +922,7 @@ namespace UE::Learning::Agents::Action::Private
 
 			return;
 		}
+
 		}
 	}
 
@@ -997,6 +998,17 @@ namespace UE::Learning::Agents::Action::Private
 		default: checkNoEntry(); return Learning::Action::EEncodingActivationFunction::ReLU;
 		}
 	}
+
+	static inline Learning::Action::ENormalization GetNormalizationOption(const ELearningAgentsNormalization NormalizationOption)
+	{
+		switch (NormalizationOption)
+		{
+		case ELearningAgentsNormalization::Manual: return Learning::Action::ENormalization::Manual;
+		case ELearningAgentsNormalization::AutoShared: return Learning::Action::ENormalization::AutoShared;
+		case ELearningAgentsNormalization::AutoPerDimension: return Learning::Action::ENormalization::AutoPerDimension;
+		default: checkNoEntry(); return Learning::Action::ENormalization::Manual;
+		}
+	}
 }
 
 bool ULearningAgentsActions::ValidateActionObjectMatchesSchema(
@@ -1060,7 +1072,7 @@ FLearningAgentsActionSchemaElement ULearningAgentsActions::SpecifyNullAction(ULe
 	return { Schema->ActionSchema.CreateNull(Tag) };
 }
 
-FLearningAgentsActionSchemaElement ULearningAgentsActions::SpecifyContinuousAction(ULearningAgentsActionSchema* Schema, const int32 Size, const float Scale, const bool bActive, const FName Tag)
+FLearningAgentsActionSchemaElement ULearningAgentsActions::SpecifyContinuousAction(ULearningAgentsActionSchema* Schema, const int32 Size, const FLearningAgentsActionNormalizationSettings& NormalizationSettings, const bool bActive, const FName Tag)
 {
 	if (!bActive)
 	{
@@ -1072,25 +1084,19 @@ FLearningAgentsActionSchemaElement ULearningAgentsActions::SpecifyContinuousActi
 		UE_LOG(LogLearning, Error, TEXT("SpecifyContinuousAction: Schema is nullptr."));
 		return FLearningAgentsActionSchemaElement();
 	}
-	
+
 	if (Size < 0)
 	{
 		UE_LOG(LogLearning, Error, TEXT("SpecifyContinuousAction: Invalid Continuous Action Size '%i'."), Size);
 		return FLearningAgentsActionSchemaElement();
 	}
 
-	if (Scale < UE_SMALL_NUMBER)
-	{
-		UE_LOG(LogLearning, Error, TEXT("SpecifyContinuousAction: Invalid Scale for Continuous Action '%f', must be greater than '%f'."), Scale, UE_SMALL_NUMBER);
-		return FLearningAgentsActionSchemaElement();
-	}
-
 	if (Size == 0)
 	{
 		UE_LOG(LogLearning, Warning, TEXT("SpecifyContinuousAction: Specifying zero-sized Continuous Action."));
 	}
 
-	return { Schema->ActionSchema.CreateContinuous({ Size, Scale }, Tag) };
+	return { Schema->ActionSchema.CreateContinuous({ Size, NormalizationSettings.Scales, NormalizationSettings.Offsets, UE::Learning::Agents::Action::Private::GetNormalizationOption(NormalizationSettings.NormalizationOption)}, Tag) };
 }
 
 FLearningAgentsActionSchemaElement ULearningAgentsActions::SpecifyExclusiveDiscreteAction(ULearningAgentsActionSchema* Schema, const int32 Size, const TArray<float>& PriorProbabilities, const bool bActive, const FName Tag)
@@ -1141,7 +1147,7 @@ FLearningAgentsActionSchemaElement ULearningAgentsActions::SpecifyExclusiveDiscr
 	}
 	UE::Learning::Agents::Action::Private::NormalizeProbabilitiesExclusive(NormalizedPriorProbabilities);
 
-	return { Schema->ActionSchema.CreateDiscreteExclusive({ Size, MakeArrayView(NormalizedPriorProbabilities) }, Tag)};
+	return { Schema->ActionSchema.CreateDiscreteExclusive({ Size, MakeArrayView(NormalizedPriorProbabilities) }, Tag) };
 }
 
 FLearningAgentsActionSchemaElement ULearningAgentsActions::SpecifyNamedExclusiveDiscreteAction(ULearningAgentsActionSchema* Schema, const TArray<FName>& Names, const TMap<FName, float>& PriorProbabilities, const bool bActive, const FName Tag)
@@ -1411,9 +1417,9 @@ FLearningAgentsActionSchemaElement ULearningAgentsActions::SpecifyStructAction(U
 	// Sort Elements According to FName
 
 	SubElementIndices.Sort([SubElementNames](const int32 Lhs, const int32 Rhs)
-	{
-		return SubElementNames[Lhs].ToString().ToLower() < SubElementNames[Rhs].ToString().ToLower();
-	});
+		{
+			return SubElementNames[Lhs].ToString().ToLower() < SubElementNames[Rhs].ToString().ToLower();
+		});
 
 	TArray<FName, TInlineAllocator<16>> SortedSubElementNames;
 	TArray<FLearningAgentsActionSchemaElement, TInlineAllocator<16>> SortedSubElements;
@@ -1522,9 +1528,9 @@ FLearningAgentsActionSchemaElement ULearningAgentsActions::SpecifyExclusiveUnion
 	// Sort Elements According to FName
 
 	SubElementIndices.Sort([SubElementNames](const int32 Lhs, const int32 Rhs)
-	{
-		return SubElementNames[Lhs].ToString().ToLower() < SubElementNames[Rhs].ToString().ToLower();
-	});
+		{
+			return SubElementNames[Lhs].ToString().ToLower() < SubElementNames[Rhs].ToString().ToLower();
+		});
 
 	TArray<FName, TInlineAllocator<16>> SortedSubElementNames;
 	TArray<FLearningAgentsActionSchemaElement, TInlineAllocator<16>> SortedSubElements;
@@ -1610,7 +1616,7 @@ FLearningAgentsActionSchemaElement ULearningAgentsActions::SpecifyExclusiveUnion
 	}
 	UE::Learning::Agents::Action::Private::NormalizeProbabilitiesExclusive(NormalizedPriorProbabilities);
 
-	return { Schema->ActionSchema.CreateOrExclusive({ ElementNames, SubElements, MakeArrayView(NormalizedPriorProbabilities) }, Tag)};
+	return { Schema->ActionSchema.CreateOrExclusive({ ElementNames, SubElements, MakeArrayView(NormalizedPriorProbabilities) }, Tag) };
 }
 
 FLearningAgentsActionSchemaElement ULearningAgentsActions::SpecifyInclusiveUnionAction(ULearningAgentsActionSchema* Schema, const TMap<FName, FLearningAgentsActionSchemaElement>& Elements, const TMap<FName, float>& PriorProbabilities, const bool bActive, const FName Tag)
@@ -1650,9 +1656,9 @@ FLearningAgentsActionSchemaElement ULearningAgentsActions::SpecifyInclusiveUnion
 	// Sort Elements According to FName
 
 	SubElementIndices.Sort([SubElementNames](const int32 Lhs, const int32 Rhs)
-	{
-		return SubElementNames[Lhs].ToString().ToLower() < SubElementNames[Rhs].ToString().ToLower();
-	});
+		{
+			return SubElementNames[Lhs].ToString().ToLower() < SubElementNames[Rhs].ToString().ToLower();
+		});
 
 	TArray<FName, TInlineAllocator<16>> SortedSubElementNames;
 	TArray<FLearningAgentsActionSchemaElement, TInlineAllocator<16>> SortedSubElements;
@@ -1738,7 +1744,7 @@ FLearningAgentsActionSchemaElement ULearningAgentsActions::SpecifyInclusiveUnion
 	}
 	UE::Learning::Agents::Action::Private::NormalizeProbabilitiesInclusive(NormalizedPriorProbabilities);
 
-	return { Schema->ActionSchema.CreateOrInclusive({ ElementNames, SubElements, MakeArrayView(NormalizedPriorProbabilities) }, Tag)};
+	return { Schema->ActionSchema.CreateOrInclusive({ ElementNames, SubElements, MakeArrayView(NormalizedPriorProbabilities) }, Tag) };
 }
 
 FLearningAgentsActionSchemaElement ULearningAgentsActions::SpecifyStaticArrayAction(ULearningAgentsActionSchema* Schema, const FLearningAgentsActionSchemaElement Element, const int32 Num, const bool bActive, const FName Tag)
@@ -1965,7 +1971,7 @@ FLearningAgentsActionSchemaElement ULearningAgentsActions::SpecifyEncodingAction
 		return FLearningAgentsActionSchemaElement();
 	}
 
-	return { Schema->ActionSchema.CreateEncoding({ Element.SchemaElement, EncodingSize, HiddenLayerNum, UE::Learning::Agents::Action::Private::GetEncodingActivationFunction(ActivationFunction) }, Tag)};
+	return { Schema->ActionSchema.CreateEncoding({ Element.SchemaElement, EncodingSize, HiddenLayerNum, UE::Learning::Agents::Action::Private::GetEncodingActivationFunction(ActivationFunction) }, Tag) };
 }
 
 FLearningAgentsActionSchemaElement ULearningAgentsActions::SpecifyBoolAction(ULearningAgentsActionSchema* Schema, const float PriorProbability, const bool bActive, const FName Tag)
@@ -1978,51 +1984,71 @@ FLearningAgentsActionSchemaElement ULearningAgentsActions::SpecifyBoolAction(ULe
 	return SpecifyExclusiveDiscreteActionFromArrayView(Schema, 2, { 1.0f - PriorProbability, PriorProbability }, bActive, Tag);
 }
 
-FLearningAgentsActionSchemaElement ULearningAgentsActions::SpecifyFloatAction(ULearningAgentsActionSchema* Schema, const float FloatScale, const bool bActive, const FName Tag)
+FLearningAgentsActionSchemaElement ULearningAgentsActions::SpecifyFloatAction(ULearningAgentsActionSchema* Schema, const FLearningAgentsFloatActionNormalizationSettings& NormalizationSettings, const bool bActive, const FName Tag)
 {
 	if (!bActive)
 	{
 		return SpecifyNullAction(Schema, Tag);
 	}
 
-	return SpecifyContinuousAction(Schema, 1, FloatScale, bActive, Tag);
+	FLearningAgentsActionNormalizationSettings Settings;
+	Settings.Scales = { NormalizationSettings.Scale };
+	Settings.Offsets = { NormalizationSettings.Offset };
+	Settings.NormalizationOption = { NormalizationSettings.NormalizationOption };
+
+	return SpecifyContinuousAction(Schema, 1, Settings, bActive, Tag);
 }
 
-FLearningAgentsActionSchemaElement ULearningAgentsActions::SpecifyLocationAction(ULearningAgentsActionSchema* Schema, const float LocationScale, const bool bActive, const FName Tag)
+FLearningAgentsActionSchemaElement ULearningAgentsActions::SpecifyLocationAction(ULearningAgentsActionSchema* Schema, const FLearningAgentsVectorActionNormalizationSettings& NormalizationSettings, const bool bActive, const FName Tag)
 {
 	if (!bActive)
 	{
 		return SpecifyNullAction(Schema, Tag);
 	}
 
-	return SpecifyContinuousAction(Schema, 3, LocationScale, bActive, Tag);
+	FLearningAgentsActionNormalizationSettings Settings;
+	Settings.Scales = { (float)NormalizationSettings.Scales.X, (float)NormalizationSettings.Scales.Y, (float)NormalizationSettings.Scales.Z };
+	Settings.Offsets = { (float)NormalizationSettings.Offsets.X, (float)NormalizationSettings.Offsets.Y, (float)NormalizationSettings.Offsets.Z };
+	Settings.NormalizationOption = { NormalizationSettings.NormalizationOption };
+
+	return SpecifyContinuousAction(Schema, 3, Settings, bActive, Tag);
 }
 
-FLearningAgentsActionSchemaElement ULearningAgentsActions::SpecifyRotationAction(ULearningAgentsActionSchema* Schema, const float RotationScale, const bool bActive, const FName Tag)
+FLearningAgentsActionSchemaElement ULearningAgentsActions::SpecifyRotationAction(ULearningAgentsActionSchema* Schema, const FLearningAgentsVectorActionNormalizationSettings& NormalizationSettings, const bool bActive, const FName Tag)
 {
 	if (!bActive)
 	{
 		return SpecifyNullAction(Schema, Tag);
 	}
 
-	return SpecifyContinuousAction(Schema, 3, RotationScale, bActive, Tag);
+	FLearningAgentsActionNormalizationSettings Settings;
+	Settings.Scales = { (float)NormalizationSettings.Scales.X, (float)NormalizationSettings.Scales.Y, (float)NormalizationSettings.Scales.Z };
+	Settings.Offsets = { (float)NormalizationSettings.Offsets.X, (float)NormalizationSettings.Offsets.Y, (float)NormalizationSettings.Offsets.Z };
+	Settings.NormalizationOption = { NormalizationSettings.NormalizationOption };
+
+	return SpecifyContinuousAction(Schema, 3, Settings, bActive, Tag);
 }
 
-FLearningAgentsActionSchemaElement ULearningAgentsActions::SpecifyScaleAction(ULearningAgentsActionSchema* Schema, const float ScaleScale, const bool bActive, const FName Tag)
+FLearningAgentsActionSchemaElement ULearningAgentsActions::SpecifyScaleAction(ULearningAgentsActionSchema* Schema, const FLearningAgentsVectorActionNormalizationSettings& NormalizationSettings, const bool bActive, const FName Tag)
 {
 	if (!bActive)
 	{
 		return SpecifyNullAction(Schema, Tag);
 	}
 
-	return SpecifyContinuousAction(Schema, 3, ScaleScale, bActive, Tag);
+	FLearningAgentsActionNormalizationSettings Settings;
+	Settings.Scales = { (float)NormalizationSettings.Scales.X, (float)NormalizationSettings.Scales.Y, (float)NormalizationSettings.Scales.Z };
+	Settings.Offsets = { (float)NormalizationSettings.Offsets.X, (float)NormalizationSettings.Offsets.Y, (float)NormalizationSettings.Offsets.Z };
+	Settings.NormalizationOption = NormalizationSettings.NormalizationOption;
+
+	return SpecifyContinuousAction(Schema, 3, Settings, bActive, Tag);
 }
 
 FLearningAgentsActionSchemaElement ULearningAgentsActions::SpecifyTransformAction(
-	ULearningAgentsActionSchema* Schema, 
-	const float LocationScale ,
-	const float RotationScale,
-	const float ScaleScale,
+	ULearningAgentsActionSchema* Schema,
+	const FLearningAgentsVectorActionNormalizationSettings& LocationNormalization,
+	const FLearningAgentsVectorActionNormalizationSettings& RotationNormalization,
+	const FLearningAgentsVectorActionNormalizationSettings& ScaleNormalization,
 	const bool bActive,
 	const FName Tag)
 {
@@ -2038,31 +2064,41 @@ FLearningAgentsActionSchemaElement ULearningAgentsActions::SpecifyTransformActio
 			TEXT("Scale")
 		},
 			{
-				SpecifyLocationAction(Schema, LocationScale),
-				SpecifyRotationAction(Schema, RotationScale),
-				SpecifyScaleAction(Schema, ScaleScale)
+				SpecifyLocationAction(Schema, LocationNormalization),
+				SpecifyRotationAction(Schema, RotationNormalization),
+				SpecifyScaleAction(Schema, ScaleNormalization)
 			},
 		bActive, Tag);
 }
 
-FLearningAgentsActionSchemaElement ULearningAgentsActions::SpecifyAngleAction(ULearningAgentsActionSchema* Schema, const float AngleScale, const bool bActive, const FName Tag)
+FLearningAgentsActionSchemaElement ULearningAgentsActions::SpecifyAngleAction(ULearningAgentsActionSchema* Schema, const FLearningAgentsFloatActionNormalizationSettings& NormalizationSettings, const bool bActive, const FName Tag)
 {
 	if (!bActive)
 	{
 		return SpecifyNullAction(Schema, Tag);
 	}
 
-	return SpecifyContinuousAction(Schema, 1, AngleScale, bActive, Tag);
+	FLearningAgentsActionNormalizationSettings Settings;
+	Settings.Scales = { NormalizationSettings.Scale };
+	Settings.Offsets = { NormalizationSettings.Offset };
+	Settings.NormalizationOption = NormalizationSettings.NormalizationOption;
+
+	return SpecifyContinuousAction(Schema, 1, Settings, bActive, Tag);
 }
 
-FLearningAgentsActionSchemaElement ULearningAgentsActions::SpecifyVelocityAction(ULearningAgentsActionSchema* Schema, const float VelocityScale, const bool bActive, const FName Tag)
+FLearningAgentsActionSchemaElement ULearningAgentsActions::SpecifyVelocityAction(ULearningAgentsActionSchema* Schema, const FLearningAgentsVectorActionNormalizationSettings& NormalizationSettings, const bool bActive, const FName Tag)
 {
 	if (!bActive)
 	{
 		return SpecifyNullAction(Schema, Tag);
 	}
 
-	return SpecifyContinuousAction(Schema, 3, VelocityScale, bActive, Tag);
+	FLearningAgentsActionNormalizationSettings Settings;
+	Settings.Scales = { (float)NormalizationSettings.Scales.X, (float)NormalizationSettings.Scales.Y, (float)NormalizationSettings.Scales.Z };
+	Settings.Offsets = { (float)NormalizationSettings.Offsets.X, (float)NormalizationSettings.Offsets.Y, (float)NormalizationSettings.Offsets.Z };
+	Settings.NormalizationOption = NormalizationSettings.NormalizationOption;
+
+	return SpecifyContinuousAction(Schema, 3, Settings, bActive, Tag);
 }
 
 FLearningAgentsActionSchemaElement ULearningAgentsActions::SpecifyDirectionAction(ULearningAgentsActionSchema* Schema, const bool bActive, const FName Tag)
@@ -2072,7 +2108,7 @@ FLearningAgentsActionSchemaElement ULearningAgentsActions::SpecifyDirectionActio
 		return SpecifyNullAction(Schema, Tag);
 	}
 
-	return SpecifyContinuousAction(Schema, 3, 1.0f, bActive, Tag);
+	return SpecifyContinuousAction(Schema, 3, FLearningAgentsActionNormalizationSettings(), bActive, Tag);
 }
 
 void ULearningAgentsActions::LogAction(const ULearningAgentsActionObject* Object, const FLearningAgentsActionObjectElement Element)
@@ -2100,32 +2136,24 @@ FLearningAgentsActionObjectElement ULearningAgentsActions::MakeNullAction(ULearn
 FLearningAgentsActionObjectElement ULearningAgentsActions::MakeContinuousAction(
 	ULearningAgentsActionObject* Object,
 	const TArray<float>& Values,
-    const bool bActive,	
+	const bool bActive,
 	const FName Tag,
-	const bool bVisualLoggerEnabled,
-	ULearningAgentsManagerListener* VisualLoggerListener,
-	const int32 VisualLoggerAgentId,
-	const FVector VisualLoggerLocation,
-	const FLinearColor VisualLoggerColor)
+	const FLearningAgentsVisualLoggerSettings& VisualLoggerSettings)
 {
 	if (!bActive)
 	{
 		return MakeNullAction(Object, Tag);
 	}
 
-	return MakeContinuousActionFromArrayView(Object, Values, bActive, Tag, bVisualLoggerEnabled, VisualLoggerListener, VisualLoggerAgentId, VisualLoggerLocation, VisualLoggerColor);
+	return MakeContinuousActionFromArrayView(Object, Values, bActive, Tag, VisualLoggerSettings);
 }
 
 FLearningAgentsActionObjectElement ULearningAgentsActions::MakeContinuousActionFromArrayView(
 	ULearningAgentsActionObject* Object,
 	const TArrayView<const float> Values,
-    const bool bActive,
+	const bool bActive,
 	const FName Tag,
-	const bool bVisualLoggerEnabled,
-	ULearningAgentsManagerListener* VisualLoggerListener,
-	const int32 VisualLoggerAgentId,
-	const FVector VisualLoggerLocation,
-	const FLinearColor VisualLoggerColor)
+	const FLearningAgentsVisualLoggerSettings& VisualLoggerSettings)
 {
 	if (!bActive)
 	{
@@ -2144,16 +2172,16 @@ FLearningAgentsActionObjectElement ULearningAgentsActions::MakeContinuousActionF
 	}
 
 #if UE_LEARNING_AGENTS_ENABLE_VISUAL_LOG
-	if (bVisualLoggerEnabled && VisualLoggerListener)
+	if (VisualLoggerSettings.bEnabled && VisualLoggerSettings.Listener)
 	{
-		const ULearningAgentsVisualLoggerObject* VisualLoggerObject = VisualLoggerListener->GetOrAddVisualLoggerObject(Tag);
+		const ULearningAgentsVisualLoggerObject* VisualLoggerObject = VisualLoggerSettings.Listener->GetOrAddVisualLoggerObject(Tag);
 
-		UE_LEARNING_AGENTS_VLOG_STRING(VisualLoggerObject, LogLearning, Display, VisualLoggerLocation,
-			VisualLoggerColor.ToFColor(true),
+		UE_LEARNING_AGENTS_VLOG_STRING(VisualLoggerObject, LogLearning, Display, VisualLoggerSettings.Location,
+			VisualLoggerSettings.Color.ToFColor(true),
 			TEXT("Listener: %s\nTag: %s\nAgent Id: % 3i\nEncoded: %s\nValues: %s"),
-			*VisualLoggerListener->GetName(),
+			*VisualLoggerSettings.Listener->GetName(),
 			*Tag.ToString(),
-			VisualLoggerAgentId,
+			VisualLoggerSettings.AgentId,
 			*UE::Learning::Array::FormatFloat(Values),
 			*UE::Learning::Array::FormatFloat(Values)); // Encoded is identical to provided values
 	}
@@ -2165,13 +2193,9 @@ FLearningAgentsActionObjectElement ULearningAgentsActions::MakeContinuousActionF
 FLearningAgentsActionObjectElement ULearningAgentsActions::MakeExclusiveDiscreteAction(
 	ULearningAgentsActionObject* Object,
 	const int32 Index,
-    const bool bActive,
+	const bool bActive,
 	const FName Tag,
-	const bool bVisualLoggerEnabled,
-	ULearningAgentsManagerListener* VisualLoggerListener,
-	const int32 VisualLoggerAgentId,
-	const FVector VisualLoggerLocation,
-	const FLinearColor VisualLoggerColor)
+	const FLearningAgentsVisualLoggerSettings& VisualLoggerSettings)
 {
 	if (!bActive)
 	{
@@ -2191,16 +2215,16 @@ FLearningAgentsActionObjectElement ULearningAgentsActions::MakeExclusiveDiscrete
 	}
 
 #if UE_LEARNING_AGENTS_ENABLE_VISUAL_LOG
-	if (bVisualLoggerEnabled && VisualLoggerListener)
+	if (VisualLoggerSettings.bEnabled && VisualLoggerSettings.Listener)
 	{
-		const ULearningAgentsVisualLoggerObject* VisualLoggerObject = VisualLoggerListener->GetOrAddVisualLoggerObject(Tag);
+		const ULearningAgentsVisualLoggerObject* VisualLoggerObject = VisualLoggerSettings.Listener->GetOrAddVisualLoggerObject(Tag);
 
-		UE_LEARNING_AGENTS_VLOG_STRING(VisualLoggerObject, LogLearning, Display, VisualLoggerLocation,
-			VisualLoggerColor.ToFColor(true),
+		UE_LEARNING_AGENTS_VLOG_STRING(VisualLoggerObject, LogLearning, Display, VisualLoggerSettings.Location,
+			VisualLoggerSettings.Color.ToFColor(true),
 			TEXT("Listener: %s\nTag: %s\nAgent Id: % 3i\nIndex: [%i]"),
-			*VisualLoggerListener->GetName(),
+			*VisualLoggerSettings.Listener->GetName(),
 			*Tag.ToString(),
-			VisualLoggerAgentId,
+			VisualLoggerSettings.AgentId,
 			Index);
 	}
 #endif
@@ -2211,13 +2235,9 @@ FLearningAgentsActionObjectElement ULearningAgentsActions::MakeExclusiveDiscrete
 FLearningAgentsActionObjectElement ULearningAgentsActions::MakeNamedExclusiveDiscreteAction(
 	ULearningAgentsActionObject* Object,
 	const FName Name,
-    const bool bActive,	
+	const bool bActive,
 	const FName Tag,
-	const bool bVisualLoggerEnabled,
-	ULearningAgentsManagerListener* VisualLoggerListener,
-	const int32 VisualLoggerAgentId,
-	const FVector VisualLoggerLocation,
-	const FLinearColor VisualLoggerColor)
+	const FLearningAgentsVisualLoggerSettings& VisualLoggerSettings)
 {
 	if (!bActive)
 	{
@@ -2231,16 +2251,16 @@ FLearningAgentsActionObjectElement ULearningAgentsActions::MakeNamedExclusiveDis
 	}
 
 #if UE_LEARNING_AGENTS_ENABLE_VISUAL_LOG
-	if (bVisualLoggerEnabled && VisualLoggerListener)
+	if (VisualLoggerSettings.bEnabled && VisualLoggerSettings.Listener)
 	{
-		const ULearningAgentsVisualLoggerObject* VisualLoggerObject = VisualLoggerListener->GetOrAddVisualLoggerObject(Tag);
+		const ULearningAgentsVisualLoggerObject* VisualLoggerObject = VisualLoggerSettings.Listener->GetOrAddVisualLoggerObject(Tag);
 
-		UE_LEARNING_AGENTS_VLOG_STRING(VisualLoggerObject, LogLearning, Display, VisualLoggerLocation,
-			VisualLoggerColor.ToFColor(true),
+		UE_LEARNING_AGENTS_VLOG_STRING(VisualLoggerObject, LogLearning, Display, VisualLoggerSettings.Location,
+			VisualLoggerSettings.Color.ToFColor(true),
 			TEXT("Listener: %s\nTag: %s\nAgent Id: % 3i\nName: [%s]"),
-			*VisualLoggerListener->GetName(),
+			*VisualLoggerSettings.Listener->GetName(),
 			*Tag.ToString(),
-			VisualLoggerAgentId,
+			VisualLoggerSettings.AgentId,
 			*Name.ToString());
 	}
 #endif
@@ -2251,32 +2271,24 @@ FLearningAgentsActionObjectElement ULearningAgentsActions::MakeNamedExclusiveDis
 FLearningAgentsActionObjectElement ULearningAgentsActions::MakeInclusiveDiscreteAction(
 	ULearningAgentsActionObject* Object,
 	const TArray<int32>& Indices,
-    const bool bActive,
+	const bool bActive,
 	const FName Tag,
-	const bool bVisualLoggerEnabled,
-	ULearningAgentsManagerListener* VisualLoggerListener,
-	const int32 VisualLoggerAgentId,
-	const FVector VisualLoggerLocation,
-	const FLinearColor VisualLoggerColor)
+	const FLearningAgentsVisualLoggerSettings& VisualLoggerSettings)
 {
 	if (!bActive)
 	{
 		return MakeNullAction(Object, Tag);
 	}
 
-	return MakeInclusiveDiscreteActionFromArrayView(Object, Indices, bActive, Tag, bVisualLoggerEnabled, VisualLoggerListener, VisualLoggerAgentId, VisualLoggerLocation, VisualLoggerColor);
+	return MakeInclusiveDiscreteActionFromArrayView(Object, Indices, bActive, Tag, VisualLoggerSettings);
 }
 
 FLearningAgentsActionObjectElement ULearningAgentsActions::MakeInclusiveDiscreteActionFromArrayView(
 	ULearningAgentsActionObject* Object,
 	const TArrayView<const int32> Indices,
-    const bool bActive,	
+	const bool bActive,
 	const FName Tag,
-	const bool bVisualLoggerEnabled,
-	ULearningAgentsManagerListener* VisualLoggerListener,
-	const int32 VisualLoggerAgentId,
-	const FVector VisualLoggerLocation,
-	const FLinearColor VisualLoggerColor)
+	const FLearningAgentsVisualLoggerSettings& VisualLoggerSettings)
 {
 	if (!bActive)
 	{
@@ -2307,16 +2319,16 @@ FLearningAgentsActionObjectElement ULearningAgentsActions::MakeInclusiveDiscrete
 	}
 
 #if UE_LEARNING_AGENTS_ENABLE_VISUAL_LOG
-	if (bVisualLoggerEnabled && VisualLoggerListener)
+	if (VisualLoggerSettings.bEnabled && VisualLoggerSettings.Listener)
 	{
-		const ULearningAgentsVisualLoggerObject* VisualLoggerObject = VisualLoggerListener->GetOrAddVisualLoggerObject(Tag);
+		const ULearningAgentsVisualLoggerObject* VisualLoggerObject = VisualLoggerSettings.Listener->GetOrAddVisualLoggerObject(Tag);
 
-		UE_LEARNING_AGENTS_VLOG_STRING(VisualLoggerObject, LogLearning, Display, VisualLoggerLocation,
-			VisualLoggerColor.ToFColor(true),
+		UE_LEARNING_AGENTS_VLOG_STRING(VisualLoggerObject, LogLearning, Display, VisualLoggerSettings.Location,
+			VisualLoggerSettings.Color.ToFColor(true),
 			TEXT("Listener: %s\nTag: %s\nAgent Id: % 3i\nIndices: %s"),
-			*VisualLoggerListener->GetName(),
+			*VisualLoggerSettings.Listener->GetName(),
 			*Tag.ToString(),
-			VisualLoggerAgentId,
+			VisualLoggerSettings.AgentId,
 			*UE::Learning::Array::FormatInt32(Indices, 256));
 	}
 #endif
@@ -2327,32 +2339,24 @@ FLearningAgentsActionObjectElement ULearningAgentsActions::MakeInclusiveDiscrete
 FLearningAgentsActionObjectElement ULearningAgentsActions::MakeNamedInclusiveDiscreteAction(
 	ULearningAgentsActionObject* Object,
 	const TArray<FName>& Names,
-    const bool bActive,	
+	const bool bActive,
 	const FName Tag,
-	const bool bVisualLoggerEnabled,
-	ULearningAgentsManagerListener* VisualLoggerListener,
-	const int32 VisualLoggerAgentId,
-	const FVector VisualLoggerLocation,
-	const FLinearColor VisualLoggerColor)
+	const FLearningAgentsVisualLoggerSettings& VisualLoggerSettings)
 {
 	if (!bActive)
 	{
 		return MakeNullAction(Object, Tag);
 	}
 
-	return MakeNamedInclusiveDiscreteActionFromArrayView(Object, Names, bActive, Tag, bVisualLoggerEnabled, VisualLoggerListener, VisualLoggerAgentId, VisualLoggerLocation, VisualLoggerColor);
+	return MakeNamedInclusiveDiscreteActionFromArrayView(Object, Names, bActive, Tag, VisualLoggerSettings);
 }
 
 FLearningAgentsActionObjectElement ULearningAgentsActions::MakeNamedInclusiveDiscreteActionFromArrayView(
 	ULearningAgentsActionObject* Object,
 	const TArrayView<const FName> Names,
-    const bool bActive,
+	const bool bActive,
 	const FName Tag,
-	const bool bVisualLoggerEnabled,
-	ULearningAgentsManagerListener* VisualLoggerListener,
-	const int32 VisualLoggerAgentId,
-	const FVector VisualLoggerLocation,
-	const FLinearColor VisualLoggerColor)
+	const FLearningAgentsVisualLoggerSettings& VisualLoggerSettings)
 {
 	if (!bActive)
 	{
@@ -2372,16 +2376,16 @@ FLearningAgentsActionObjectElement ULearningAgentsActions::MakeNamedInclusiveDis
 	}
 
 #if UE_LEARNING_AGENTS_ENABLE_VISUAL_LOG
-	if (bVisualLoggerEnabled && VisualLoggerListener)
+	if (VisualLoggerSettings.bEnabled && VisualLoggerSettings.Listener)
 	{
-		const ULearningAgentsVisualLoggerObject* VisualLoggerObject = VisualLoggerListener->GetOrAddVisualLoggerObject(Tag);
+		const ULearningAgentsVisualLoggerObject* VisualLoggerObject = VisualLoggerSettings.Listener->GetOrAddVisualLoggerObject(Tag);
 
-		UE_LEARNING_AGENTS_VLOG_STRING(VisualLoggerObject, LogLearning, Display, VisualLoggerLocation,
-			VisualLoggerColor.ToFColor(true),
+		UE_LEARNING_AGENTS_VLOG_STRING(VisualLoggerObject, LogLearning, Display, VisualLoggerSettings.Location,
+			VisualLoggerSettings.Color.ToFColor(true),
 			TEXT("Listener: %s\nTag: %s\nAgent Id: % 3i\nNames: %s"),
-			*VisualLoggerListener->GetName(),
+			*VisualLoggerSettings.Listener->GetName(),
 			*Tag.ToString(),
-			VisualLoggerAgentId,
+			VisualLoggerSettings.AgentId,
 			*UE::Learning::Agents::Action::Private::FormatNames(Names, 256));
 	}
 #endif
@@ -2630,13 +2634,9 @@ FLearningAgentsActionObjectElement ULearningAgentsActions::MakeEnumAction(
 	ULearningAgentsActionObject* Object,
 	const UEnum* Enum,
 	const uint8 EnumValue,
-    const bool bActive,
+	const bool bActive,
 	const FName Tag,
-	const bool bVisualLoggerEnabled,
-	ULearningAgentsManagerListener* VisualLoggerListener,
-	const int32 VisualLoggerAgentId,
-	const FVector VisualLoggerLocation,
-	const FLinearColor VisualLoggerColor)
+	const FLearningAgentsVisualLoggerSettings& VisualLoggerSettings)
 {
 	if (!bActive)
 	{
@@ -2653,21 +2653,21 @@ FLearningAgentsActionObjectElement ULearningAgentsActions::MakeEnumAction(
 
 	if (EnumValueIndex == INDEX_NONE || EnumValueIndex < 0 || EnumValueIndex >= Enum->NumEnums() - 1)
 	{
-		UE_LOG(LogLearning, Error, TEXT("MakeEnumAction: EnumValue %i not valid for Enum '%s'."), EnumValue , *Enum->GetName());
+		UE_LOG(LogLearning, Error, TEXT("MakeEnumAction: EnumValue %i not valid for Enum '%s'."), EnumValue, *Enum->GetName());
 		return FLearningAgentsActionObjectElement();
 	}
 
 #if UE_LEARNING_AGENTS_ENABLE_VISUAL_LOG
-	if (bVisualLoggerEnabled && VisualLoggerListener)
+	if (VisualLoggerSettings.bEnabled && VisualLoggerSettings.Listener)
 	{
-		const ULearningAgentsVisualLoggerObject* VisualLoggerObject = VisualLoggerListener->GetOrAddVisualLoggerObject(Tag);
+		const ULearningAgentsVisualLoggerObject* VisualLoggerObject = VisualLoggerSettings.Listener->GetOrAddVisualLoggerObject(Tag);
 
-		UE_LEARNING_AGENTS_VLOG_STRING(VisualLoggerObject, LogLearning, Display, VisualLoggerLocation,
-			VisualLoggerColor.ToFColor(true),
+		UE_LEARNING_AGENTS_VLOG_STRING(VisualLoggerObject, LogLearning, Display, VisualLoggerSettings.Location,
+			VisualLoggerSettings.Color.ToFColor(true),
 			TEXT("Listener: %s\nTag: %s\nAgent Id: % 3i\nEnum: %s\nSize: [%i]\nValue: [%s]\nIndex: [%i]"),
-			*VisualLoggerListener->GetName(),
+			*VisualLoggerSettings.Listener->GetName(),
 			*Tag.ToString(),
-			VisualLoggerAgentId,
+			VisualLoggerSettings.AgentId,
 			*Enum->GetName(),
 			Enum->NumEnums() - 1,
 			*Enum->GetDisplayNameTextByValue(EnumValue).ToString(),
@@ -2682,13 +2682,9 @@ FLearningAgentsActionObjectElement ULearningAgentsActions::MakeBitmaskAction(
 	ULearningAgentsActionObject* Object,
 	const UEnum* Enum,
 	const int32 BitmaskValue,
-    const bool bActive,	
+	const bool bActive,
 	const FName Tag,
-	const bool bVisualLoggerEnabled,
-	ULearningAgentsManagerListener* VisualLoggerListener,
-	const int32 VisualLoggerAgentId,
-	const FVector VisualLoggerLocation,
-	const FLinearColor VisualLoggerColor)
+	const FLearningAgentsVisualLoggerSettings& VisualLoggerSettings)
 {
 	if (!bActive)
 	{
@@ -2719,9 +2715,9 @@ FLearningAgentsActionObjectElement ULearningAgentsActions::MakeBitmaskAction(
 	}
 
 #if UE_LEARNING_AGENTS_ENABLE_VISUAL_LOG
-	if (bVisualLoggerEnabled && VisualLoggerListener)
+	if (VisualLoggerSettings.bEnabled && VisualLoggerSettings.Listener)
 	{
-		const ULearningAgentsVisualLoggerObject* VisualLoggerObject = VisualLoggerListener->GetOrAddVisualLoggerObject(Tag);
+		const ULearningAgentsVisualLoggerObject* VisualLoggerObject = VisualLoggerSettings.Listener->GetOrAddVisualLoggerObject(Tag);
 
 		FString ValuesString;
 		FString IndicesString;
@@ -2738,12 +2734,12 @@ FLearningAgentsActionObjectElement ULearningAgentsActions::MakeBitmaskAction(
 		ValuesString = ValuesString.TrimEnd();
 		IndicesString = IndicesString.TrimEnd();
 
-		UE_LEARNING_AGENTS_VLOG_STRING(VisualLoggerObject, LogLearning, Display, VisualLoggerLocation,
-			VisualLoggerColor.ToFColor(true),
+		UE_LEARNING_AGENTS_VLOG_STRING(VisualLoggerObject, LogLearning, Display, VisualLoggerSettings.Location,
+			VisualLoggerSettings.Color.ToFColor(true),
 			TEXT("Listener: %s\nTag: %s\nAgent Id: % 3i\nEnum: %s\nSize: [%i]\nValues: [%s]\nIndices: [%s]"),
-			*VisualLoggerListener->GetName(),
+			*VisualLoggerSettings.Listener->GetName(),
 			*Tag.ToString(),
-			VisualLoggerAgentId,
+			VisualLoggerSettings.AgentId,
 			*Enum->GetName(),
 			Enum->NumEnums() - 1,
 			*ValuesString,
@@ -2842,13 +2838,9 @@ FLearningAgentsActionObjectElement ULearningAgentsActions::MakeEncodingAction(UL
 FLearningAgentsActionObjectElement ULearningAgentsActions::MakeBoolAction(
 	ULearningAgentsActionObject* Object,
 	const bool bValue,
-    const bool bActive,
+	const bool bActive,
 	const FName Tag,
-	const bool bVisualLoggerEnabled,
-	ULearningAgentsManagerListener* VisualLoggerListener,
-	const int32 VisualLoggerAgentId,
-	const FVector VisualLoggerLocation,
-	const FLinearColor VisualLoggerColor)
+	const FLearningAgentsVisualLoggerSettings& VisualLoggerSettings)
 {
 	if (!bActive)
 	{
@@ -2856,16 +2848,16 @@ FLearningAgentsActionObjectElement ULearningAgentsActions::MakeBoolAction(
 	}
 
 #if UE_LEARNING_AGENTS_ENABLE_VISUAL_LOG
-	if (bVisualLoggerEnabled && VisualLoggerListener)
+	if (VisualLoggerSettings.bEnabled && VisualLoggerSettings.Listener)
 	{
-		const ULearningAgentsVisualLoggerObject* VisualLoggerObject = VisualLoggerListener->GetOrAddVisualLoggerObject(Tag);
+		const ULearningAgentsVisualLoggerObject* VisualLoggerObject = VisualLoggerSettings.Listener->GetOrAddVisualLoggerObject(Tag);
 
-		UE_LEARNING_AGENTS_VLOG_STRING(VisualLoggerObject, LogLearning, Display, VisualLoggerLocation,
-			VisualLoggerColor.ToFColor(true),
+		UE_LEARNING_AGENTS_VLOG_STRING(VisualLoggerObject, LogLearning, Display, VisualLoggerSettings.Location,
+			VisualLoggerSettings.Color.ToFColor(true),
 			TEXT("Listener: %s\nTag: %s\nAgent Id: % 3i\nValue: [%s]"),
-			*VisualLoggerListener->GetName(),
+			*VisualLoggerSettings.Listener->GetName(),
 			*Tag.ToString(),
-			VisualLoggerAgentId,
+			VisualLoggerSettings.AgentId,
 			bValue ? TEXT("true") : TEXT("false"));
 	}
 #endif
@@ -2876,13 +2868,9 @@ FLearningAgentsActionObjectElement ULearningAgentsActions::MakeBoolAction(
 FLearningAgentsActionObjectElement ULearningAgentsActions::MakeFloatAction(
 	ULearningAgentsActionObject* Object,
 	const float Value,
-    const bool bActive,
+	const bool bActive,
 	const FName Tag,
-	const bool bVisualLoggerEnabled,
-	ULearningAgentsManagerListener* VisualLoggerListener,
-	const int32 VisualLoggerAgentId,
-	const FVector VisualLoggerLocation,
-	const FLinearColor VisualLoggerColor)
+	const FLearningAgentsVisualLoggerSettings& VisualLoggerSettings)
 {
 	if (!bActive)
 	{
@@ -2890,34 +2878,30 @@ FLearningAgentsActionObjectElement ULearningAgentsActions::MakeFloatAction(
 	}
 
 #if UE_LEARNING_AGENTS_ENABLE_VISUAL_LOG
-	if (bVisualLoggerEnabled && VisualLoggerListener)
+	if (VisualLoggerSettings.bEnabled && VisualLoggerSettings.Listener)
 	{
-		const ULearningAgentsVisualLoggerObject* VisualLoggerObject = VisualLoggerListener->GetOrAddVisualLoggerObject(Tag);
+		const ULearningAgentsVisualLoggerObject* VisualLoggerObject = VisualLoggerSettings.Listener->GetOrAddVisualLoggerObject(Tag);
 
-		UE_LEARNING_AGENTS_VLOG_STRING(VisualLoggerObject, LogLearning, Display, VisualLoggerLocation,
-			VisualLoggerColor.ToFColor(true),
+		UE_LEARNING_AGENTS_VLOG_STRING(VisualLoggerObject, LogLearning, Display, VisualLoggerSettings.Location,
+			VisualLoggerSettings.Color.ToFColor(true),
 			TEXT("Listener: %s\nTag: %s\nAgent Id: % 3i\nValue: [% 6.1f]"),
-			*VisualLoggerListener->GetName(),
+			*VisualLoggerSettings.Listener->GetName(),
 			*Tag.ToString(),
-			VisualLoggerAgentId,
+			VisualLoggerSettings.AgentId,
 			Value);
 	}
 #endif
 
-	return MakeContinuousActionFromArrayView(Object, { Value }, bActive, Tag);
+	return MakeContinuousActionFromArrayView(Object, { Value }, bActive, Tag, FLearningAgentsVisualLoggerSettings());
 }
 
 FLearningAgentsActionObjectElement ULearningAgentsActions::MakeLocationAction(
 	ULearningAgentsActionObject* Object,
 	const FVector Location,
 	const FTransform RelativeTransform,
-	const bool bActive,	
+	const bool bActive,
 	const FName Tag,
-	const bool bVisualLoggerEnabled,
-	ULearningAgentsManagerListener* VisualLoggerListener,
-	const int32 VisualLoggerAgentId,
-	const FVector VisualLoggerLocation,
-	const FLinearColor VisualLoggerColor)
+	const FLearningAgentsVisualLoggerSettings& VisualLoggerSettings)
 {
 	if (!bActive)
 	{
@@ -2927,34 +2911,34 @@ FLearningAgentsActionObjectElement ULearningAgentsActions::MakeLocationAction(
 	const FVector LocalLocation = RelativeTransform.InverseTransformPosition(Location);
 
 #if UE_LEARNING_AGENTS_ENABLE_VISUAL_LOG
-	if (bVisualLoggerEnabled && VisualLoggerListener)
+	if (VisualLoggerSettings.bEnabled && VisualLoggerSettings.Listener)
 	{
-		const ULearningAgentsVisualLoggerObject* VisualLoggerObject = VisualLoggerListener->GetOrAddVisualLoggerObject(Tag);
+		const ULearningAgentsVisualLoggerObject* VisualLoggerObject = VisualLoggerSettings.Listener->GetOrAddVisualLoggerObject(Tag);
 
 		UE_LEARNING_AGENTS_VLOG_LOCATION(VisualLoggerObject, LogLearning, Display,
 			Location,
 			10,
-			VisualLoggerColor.ToFColor(true),
+			VisualLoggerSettings.Color.ToFColor(true),
 			TEXT(""));
 
 		UE_LEARNING_AGENTS_VLOG_SEGMENT(VisualLoggerObject, LogLearning, Display,
 			RelativeTransform.GetTranslation(),
 			Location,
-			VisualLoggerColor.ToFColor(true),
+			VisualLoggerSettings.Color.ToFColor(true),
 			TEXT(""));
 
 		UE_LEARNING_AGENTS_VLOG_TRANSFORM(VisualLoggerObject, LogLearning, Display,
 			RelativeTransform.GetTranslation(),
 			RelativeTransform.GetRotation(),
-			VisualLoggerColor.ToFColor(true),
+			VisualLoggerSettings.Color.ToFColor(true),
 			TEXT(""));
 
-		UE_LEARNING_AGENTS_VLOG_STRING(VisualLoggerObject, LogLearning, Display, VisualLoggerLocation,
-			VisualLoggerColor.ToFColor(true),
+		UE_LEARNING_AGENTS_VLOG_STRING(VisualLoggerObject, LogLearning, Display, VisualLoggerSettings.Location,
+			VisualLoggerSettings.Color.ToFColor(true),
 			TEXT("Listener: %s\nTag: %s\nAgent Id: % 3i\nLocal Location: [% 6.1f % 6.1f % 6.1f]\nLocation: [% 6.1f % 6.1f % 6.1f]"),
-			*VisualLoggerListener->GetName(),
+			*VisualLoggerSettings.Listener->GetName(),
 			*Tag.ToString(),
-			VisualLoggerAgentId,
+			VisualLoggerSettings.AgentId,
 			LocalLocation.X, LocalLocation.Y, LocalLocation.Z,
 			Location.X, Location.Y, Location.Z);
 	}
@@ -2963,21 +2947,16 @@ FLearningAgentsActionObjectElement ULearningAgentsActions::MakeLocationAction(
 	return MakeContinuousActionFromArrayView(Object, {
 		(float)LocalLocation.X,
 		(float)LocalLocation.Y,
-		(float)LocalLocation.Z }, bActive, Tag);
+		(float)LocalLocation.Z }, bActive, Tag, FLearningAgentsVisualLoggerSettings());
 }
 
 FLearningAgentsActionObjectElement ULearningAgentsActions::MakeRotationAction(
 	ULearningAgentsActionObject* Object,
 	const FRotator Rotation,
 	const FRotator RelativeRotation,
-    const bool bActive,	
+	const bool bActive,
 	const FName Tag,
-	const bool bVisualLoggerEnabled,
-	ULearningAgentsManagerListener* VisualLoggerListener,
-	const int32 VisualLoggerAgentId,
-	const FVector VisualLoggerRotationLocation,
-	const FVector VisualLoggerLocation,
-	const FLinearColor VisualLoggerColor)
+	const FLearningAgentsVisualLoggerSettings& VisualLoggerSettings)
 {
 	if (!bActive)
 	{
@@ -2985,7 +2964,7 @@ FLearningAgentsActionObjectElement ULearningAgentsActions::MakeRotationAction(
 	}
 
 	return MakeRotationActionFromQuat(Object, FQuat::MakeFromRotator(Rotation), FQuat::MakeFromRotator(RelativeRotation), bActive, Tag,
-		bVisualLoggerEnabled, VisualLoggerListener, VisualLoggerAgentId, VisualLoggerRotationLocation, VisualLoggerLocation, VisualLoggerColor);
+		VisualLoggerSettings);
 }
 
 FLearningAgentsActionObjectElement ULearningAgentsActions::MakeRotationActionFromQuat(
@@ -2994,12 +2973,7 @@ FLearningAgentsActionObjectElement ULearningAgentsActions::MakeRotationActionFro
 	const FQuat RelativeRotation,
 	const bool bActive,
 	const FName Tag,
-	const bool bVisualLoggerEnabled,
-	ULearningAgentsManagerListener* VisualLoggerListener,
-	const int32 VisualLoggerAgentId,
-	const FVector VisualLoggerRotationLocation,
-	const FVector VisualLoggerLocation,
-	const FLinearColor VisualLoggerColor)
+	const FLearningAgentsVisualLoggerSettings& VisualLoggerSettings)
 {
 	if (!bActive)
 	{
@@ -3010,22 +2984,22 @@ FLearningAgentsActionObjectElement ULearningAgentsActions::MakeRotationActionFro
 	const FVector RotationVector = LocalRotation.ToRotationVector();
 
 #if UE_LEARNING_AGENTS_ENABLE_VISUAL_LOG
-	if (bVisualLoggerEnabled && VisualLoggerListener)
+	if (VisualLoggerSettings.bEnabled && VisualLoggerSettings.Listener)
 	{
-		const ULearningAgentsVisualLoggerObject* VisualLoggerObject = VisualLoggerListener->GetOrAddVisualLoggerObject(Tag);
+		const ULearningAgentsVisualLoggerObject* VisualLoggerObject = VisualLoggerSettings.Listener->GetOrAddVisualLoggerObject(Tag);
 
 		UE_LEARNING_AGENTS_VLOG_TRANSFORM(VisualLoggerObject, LogLearning, Display,
-			VisualLoggerRotationLocation,
+			VisualLoggerSettings.DebugLocation,
 			LocalRotation.Rotator(),
-			VisualLoggerColor.ToFColor(true),
+			VisualLoggerSettings.Color.ToFColor(true),
 			TEXT(""));
 
-		UE_LEARNING_AGENTS_VLOG_STRING(VisualLoggerObject, LogLearning, Display, VisualLoggerLocation,
-			VisualLoggerColor.ToFColor(true),
+		UE_LEARNING_AGENTS_VLOG_STRING(VisualLoggerObject, LogLearning, Display, VisualLoggerSettings.Location,
+			VisualLoggerSettings.Color.ToFColor(true),
 			TEXT("Listener: %s\nTag: %s\nAgent Id: % 3i\nLocal Rotation Vector: [% 6.1f % 6.1f % 6.1f]\nLocal Rotation: [% 6.1f % 6.1f % 6.1f % 6.1f]\nRotation: [% 6.1f % 6.1f % 6.1f % 6.1f]"),
-			*VisualLoggerListener->GetName(),
+			*VisualLoggerSettings.Listener->GetName(),
 			*Tag.ToString(),
-			VisualLoggerAgentId,
+			VisualLoggerSettings.AgentId,
 			RotationVector.X, RotationVector.Y, RotationVector.Z,
 			LocalRotation.X, LocalRotation.Y, LocalRotation.Z, LocalRotation.W,
 			Rotation.X, Rotation.Y, Rotation.Z, Rotation.W);
@@ -3036,20 +3010,16 @@ FLearningAgentsActionObjectElement ULearningAgentsActions::MakeRotationActionFro
 		(float)FMath::RadiansToDegrees(RotationVector.X),
 		(float)FMath::RadiansToDegrees(RotationVector.Y),
 		(float)FMath::RadiansToDegrees(RotationVector.Z),
-		}, bActive, Tag);
+		}, bActive, Tag, FLearningAgentsVisualLoggerSettings());
 }
 
 FLearningAgentsActionObjectElement ULearningAgentsActions::MakeScaleAction(
 	ULearningAgentsActionObject* Object,
 	const FVector Scale,
 	const FVector RelativeScale,
-    const bool bActive,	
+	const bool bActive,
 	const FName Tag,
-	const bool bVisualLoggerEnabled,
-	ULearningAgentsManagerListener* VisualLoggerListener,
-	const int32 VisualLoggerAgentId,
-	const FVector VisualLoggerLocation,
-	const FLinearColor VisualLoggerColor)
+	const FLearningAgentsVisualLoggerSettings& VisualLoggerSettings)
 {
 	if (!bActive)
 	{
@@ -3061,16 +3031,16 @@ FLearningAgentsActionObjectElement ULearningAgentsActions::MakeScaleAction(
 		UE::Learning::Agents::Action::Private::VectorLogSafe(RelativeScale);
 
 #if UE_LEARNING_AGENTS_ENABLE_VISUAL_LOG
-	if (bVisualLoggerEnabled && VisualLoggerListener)
+	if (VisualLoggerSettings.bEnabled && VisualLoggerSettings.Listener)
 	{
-		const ULearningAgentsVisualLoggerObject* VisualLoggerObject = VisualLoggerListener->GetOrAddVisualLoggerObject(Tag);
+		const ULearningAgentsVisualLoggerObject* VisualLoggerObject = VisualLoggerSettings.Listener->GetOrAddVisualLoggerObject(Tag);
 
-		UE_LEARNING_AGENTS_VLOG_STRING(VisualLoggerObject, LogLearning, Display, VisualLoggerLocation,
-			VisualLoggerColor.ToFColor(true),
+		UE_LEARNING_AGENTS_VLOG_STRING(VisualLoggerObject, LogLearning, Display, VisualLoggerSettings.Location,
+			VisualLoggerSettings.Color.ToFColor(true),
 			TEXT("Listener: %s\nTag: %s\nAgent Id: % 3i\nLocal Log Scale: [% 6.1f % 6.1f % 6.1f]\nScale: [% 6.1f % 6.1f % 6.1f]"),
-			*VisualLoggerListener->GetName(),
+			*VisualLoggerSettings.Listener->GetName(),
 			*Tag.ToString(),
-			VisualLoggerAgentId,
+			VisualLoggerSettings.AgentId,
 			LocalLogScale.X, LocalLogScale.Y, LocalLogScale.Z,
 			Scale.X, Scale.Y, Scale.Z);
 	}
@@ -3080,20 +3050,16 @@ FLearningAgentsActionObjectElement ULearningAgentsActions::MakeScaleAction(
 		(float)LocalLogScale.X,
 		(float)LocalLogScale.Y,
 		(float)LocalLogScale.Z,
-		}, bActive, Tag);
+		}, bActive, Tag, FLearningAgentsVisualLoggerSettings());
 }
 
 FLearningAgentsActionObjectElement ULearningAgentsActions::MakeTransformAction(
 	ULearningAgentsActionObject* Object,
 	const FTransform Transform,
 	const FTransform RelativeTransform,
-	const bool bActive,	
+	const bool bActive,
 	const FName Tag,
-	const bool bVisualLoggerEnabled,
-	ULearningAgentsManagerListener* VisualLoggerListener,
-	const int32 VisualLoggerAgentId,
-	const FVector VisualLoggerLocation,
-	const FLinearColor VisualLoggerColor)
+	const FLearningAgentsVisualLoggerSettings& VisualLoggerSettings)
 {
 	if (!bActive)
 	{
@@ -3105,22 +3071,22 @@ FLearningAgentsActionObjectElement ULearningAgentsActions::MakeTransformAction(
 	FVector Scale = Transform.GetScale3D();
 
 #if UE_LEARNING_AGENTS_ENABLE_VISUAL_LOG
-	if (bVisualLoggerEnabled && VisualLoggerListener)
+	if (VisualLoggerSettings.bEnabled && VisualLoggerSettings.Listener)
 	{
-		const ULearningAgentsVisualLoggerObject* VisualLoggerObject = VisualLoggerListener->GetOrAddVisualLoggerObject(Tag);
+		const ULearningAgentsVisualLoggerObject* VisualLoggerObject = VisualLoggerSettings.Listener->GetOrAddVisualLoggerObject(Tag);
 
 		UE_LEARNING_AGENTS_VLOG_TRANSFORM(VisualLoggerObject, LogLearning, Display,
 			Location,
 			Rotation,
-			VisualLoggerColor.ToFColor(true),
+			VisualLoggerSettings.Color.ToFColor(true),
 			TEXT(""));
 
-		UE_LEARNING_AGENTS_VLOG_STRING(VisualLoggerObject, LogLearning, Display, VisualLoggerLocation,
-			VisualLoggerColor.ToFColor(true),
+		UE_LEARNING_AGENTS_VLOG_STRING(VisualLoggerObject, LogLearning, Display, VisualLoggerSettings.Location,
+			VisualLoggerSettings.Color.ToFColor(true),
 			TEXT("Listener: %s\nTag: %s\nAgent Id: % 3i\nLocation: [% 6.1f % 6.1f % 6.1f]\nRotation: [% 6.1f % 6.1f % 6.1f % 6.1f]\nScale: [% 6.1f % 6.1f % 6.1f]"),
-			*VisualLoggerListener->GetName(),
+			*VisualLoggerSettings.Listener->GetName(),
 			*Tag.ToString(),
-			VisualLoggerAgentId,
+			VisualLoggerSettings.AgentId,
 			Location.X, Location.Y, Location.Z,
 			Rotation.X, Rotation.Y, Rotation.Z, Rotation.W,
 			Scale.X, Scale.Y, Scale.Z);
@@ -3147,12 +3113,7 @@ FLearningAgentsActionObjectElement ULearningAgentsActions::MakeAngleAction(
 	const float RelativeAngle,
 	const bool bActive,
 	const FName Tag,
-	const bool bVisualLoggerEnabled,
-	ULearningAgentsManagerListener* VisualLoggerListener,
-	const int32 VisualLoggerAgentId,
-	const FVector VisualLoggerAngleLocation,
-	const FVector VisualLoggerLocation,
-	const FLinearColor VisualLoggerColor)
+	const FLearningAgentsVisualLoggerSettings& VisualLoggerSettings)
 {
 	if (!bActive)
 	{
@@ -3160,46 +3121,41 @@ FLearningAgentsActionObjectElement ULearningAgentsActions::MakeAngleAction(
 	}
 
 	float LocalAngle = FMath::FindDeltaAngleDegrees(RelativeAngle, Angle);
-	
+
 #if UE_LEARNING_AGENTS_ENABLE_VISUAL_LOG
-	if (bVisualLoggerEnabled && VisualLoggerListener)
+	if (VisualLoggerSettings.bEnabled && VisualLoggerSettings.Listener)
 	{
-		const ULearningAgentsVisualLoggerObject* VisualLoggerObject = VisualLoggerListener->GetOrAddVisualLoggerObject(Tag);
+		const ULearningAgentsVisualLoggerObject* VisualLoggerObject = VisualLoggerSettings.Listener->GetOrAddVisualLoggerObject(Tag);
 
 		UE_LEARNING_AGENTS_VLOG_ANGLE_DEGREES(VisualLoggerObject, LogLearning, Display,
 			Angle,
 			0.0f,
-			VisualLoggerAngleLocation,
+			VisualLoggerSettings.DebugLocation,
 			10.0f,
-			VisualLoggerColor.ToFColor(true),
+			VisualLoggerSettings.Color.ToFColor(true),
 			TEXT(""));
 
-		UE_LEARNING_AGENTS_VLOG_STRING(VisualLoggerObject, LogLearning, Display, VisualLoggerLocation,
-			VisualLoggerColor.ToFColor(true),
+		UE_LEARNING_AGENTS_VLOG_STRING(VisualLoggerObject, LogLearning, Display, VisualLoggerSettings.Location,
+			VisualLoggerSettings.Color.ToFColor(true),
 			TEXT("Listener: %s\nTag: %s\nAgent Id: % 3i\nLocal Angle: [% 6.1f]\nAngle: [% 6.1f]"),
-			*VisualLoggerListener->GetName(),
+			*VisualLoggerSettings.Listener->GetName(),
 			*Tag.ToString(),
-			VisualLoggerAgentId,
+			VisualLoggerSettings.AgentId,
 			LocalAngle,
 			Angle);
 	}
 #endif
 
-	return MakeContinuousActionFromArrayView(Object, { LocalAngle }, bActive, Tag);
+	return MakeContinuousActionFromArrayView(Object, { LocalAngle }, bActive, Tag, FLearningAgentsVisualLoggerSettings());
 }
 
 FLearningAgentsActionObjectElement ULearningAgentsActions::MakeAngleActionRadians(
 	ULearningAgentsActionObject* Object,
 	const float Angle,
 	const float RelativeAngle,
-    const bool bActive,	
+	const bool bActive,
 	const FName Tag,
-	const bool bVisualLoggerEnabled,
-	ULearningAgentsManagerListener* VisualLoggerListener,
-	const int32 VisualLoggerAgentId,
-	const FVector VisualLoggerAngleLocation,
-	const FVector VisualLoggerLocation,
-	const FLinearColor VisualLoggerColor)
+	const FLearningAgentsVisualLoggerSettings& VisualLoggerSettings)
 {
 	if (!bActive)
 	{
@@ -3207,21 +3163,16 @@ FLearningAgentsActionObjectElement ULearningAgentsActions::MakeAngleActionRadian
 	}
 
 	return MakeAngleAction(Object, FMath::RadiansToDegrees(Angle), FMath::RadiansToDegrees(RelativeAngle), bActive, Tag,
-		bVisualLoggerEnabled, VisualLoggerListener, VisualLoggerAgentId, VisualLoggerAngleLocation, VisualLoggerLocation, VisualLoggerColor);
+		VisualLoggerSettings);
 }
 
 FLearningAgentsActionObjectElement ULearningAgentsActions::MakeVelocityAction(
 	ULearningAgentsActionObject* Object,
 	const FVector Velocity,
 	const FTransform RelativeTransform,
-	const bool bActive,	
+	const bool bActive,
 	const FName Tag,
-	const bool bVisualLoggerEnabled,
-	ULearningAgentsManagerListener* VisualLoggerListener,
-	const int32 VisualLoggerAgentId,
-	const FVector VisualLoggerVelocityLocation,
-	const FVector VisualLoggerLocation,
-	const FLinearColor VisualLoggerColor)
+	const FLearningAgentsVisualLoggerSettings& VisualLoggerSettings)
 {
 	if (!bActive)
 	{
@@ -3231,28 +3182,28 @@ FLearningAgentsActionObjectElement ULearningAgentsActions::MakeVelocityAction(
 	const FVector LocalVelocity = RelativeTransform.InverseTransformVectorNoScale(Velocity);
 
 #if UE_LEARNING_AGENTS_ENABLE_VISUAL_LOG
-	if (bVisualLoggerEnabled && VisualLoggerListener)
+	if (VisualLoggerSettings.bEnabled && VisualLoggerSettings.Listener)
 	{
-		const ULearningAgentsVisualLoggerObject* VisualLoggerObject = VisualLoggerListener->GetOrAddVisualLoggerObject(Tag);
+		const ULearningAgentsVisualLoggerObject* VisualLoggerObject = VisualLoggerSettings.Listener->GetOrAddVisualLoggerObject(Tag);
 
 		UE_LEARNING_AGENTS_VLOG_ARROW(VisualLoggerObject, LogLearning, Display,
-			VisualLoggerVelocityLocation,
-			VisualLoggerVelocityLocation + Velocity,
-			VisualLoggerColor.ToFColor(true),
+			VisualLoggerSettings.DebugLocation,
+			VisualLoggerSettings.DebugLocation + Velocity,
+			VisualLoggerSettings.Color.ToFColor(true),
 			TEXT(""));
 
 		UE_LEARNING_AGENTS_VLOG_TRANSFORM(VisualLoggerObject, LogLearning, Display,
 			RelativeTransform.GetTranslation(),
 			RelativeTransform.GetRotation(),
-			VisualLoggerColor.ToFColor(true),
+			VisualLoggerSettings.Color.ToFColor(true),
 			TEXT(""));
 
-		UE_LEARNING_AGENTS_VLOG_STRING(VisualLoggerObject, LogLearning, Display, VisualLoggerLocation,
-			VisualLoggerColor.ToFColor(true),
+		UE_LEARNING_AGENTS_VLOG_STRING(VisualLoggerObject, LogLearning, Display, VisualLoggerSettings.Location,
+			VisualLoggerSettings.Color.ToFColor(true),
 			TEXT("Listener: %s\nTag: %s\nAgent Id: % 3i\nLocal Velocity: [% 6.1f % 6.1f % 6.1f]\nVelocity: [% 6.1f % 6.1f % 6.1f]"),
-			*VisualLoggerListener->GetName(),
+			*VisualLoggerSettings.Listener->GetName(),
 			*Tag.ToString(),
-			VisualLoggerAgentId,
+			VisualLoggerSettings.AgentId,
 			LocalVelocity.X, LocalVelocity.Y, LocalVelocity.Z,
 			Velocity.X, Velocity.Y, Velocity.Z);
 	}
@@ -3262,22 +3213,16 @@ FLearningAgentsActionObjectElement ULearningAgentsActions::MakeVelocityAction(
 		(float)LocalVelocity.X,
 		(float)LocalVelocity.Y,
 		(float)LocalVelocity.Z,
-		}, bActive, Tag);
+		}, bActive, Tag, FLearningAgentsVisualLoggerSettings());
 }
 
 FLearningAgentsActionObjectElement ULearningAgentsActions::MakeDirectionAction(
 	ULearningAgentsActionObject* Object,
 	const FVector Direction,
 	const FTransform RelativeTransform,
-	const bool bActive,	
+	const bool bActive,
 	const FName Tag,
-	const bool bVisualLoggerEnabled,
-	ULearningAgentsManagerListener* VisualLoggerListener,
-	const int32 VisualLoggerAgentId,
-	const FVector VisualLoggerDirectionLocation,
-	const FVector VisualLoggerLocation,
-	const float VisualLoggerArrowLength,
-	const FLinearColor VisualLoggerColor)
+	const FLearningAgentsVisualLoggerSettings& VisualLoggerSettings)
 {
 	if (!bActive)
 	{
@@ -3287,28 +3232,28 @@ FLearningAgentsActionObjectElement ULearningAgentsActions::MakeDirectionAction(
 	const FVector LocalDirection = RelativeTransform.InverseTransformVectorNoScale(Direction).GetSafeNormal(UE_SMALL_NUMBER, FVector::ForwardVector);
 
 #if UE_LEARNING_AGENTS_ENABLE_VISUAL_LOG
-	if (bVisualLoggerEnabled && VisualLoggerListener)
+	if (VisualLoggerSettings.bEnabled && VisualLoggerSettings.Listener)
 	{
-		const ULearningAgentsVisualLoggerObject* VisualLoggerObject = VisualLoggerListener->GetOrAddVisualLoggerObject(Tag);
+		const ULearningAgentsVisualLoggerObject* VisualLoggerObject = VisualLoggerSettings.Listener->GetOrAddVisualLoggerObject(Tag);
 
 		UE_LEARNING_AGENTS_VLOG_ARROW(VisualLoggerObject, LogLearning, Display,
-			VisualLoggerDirectionLocation,
-			VisualLoggerDirectionLocation + VisualLoggerArrowLength * Direction,
-			VisualLoggerColor.ToFColor(true),
+			VisualLoggerSettings.DebugLocation,
+			VisualLoggerSettings.DebugLocation + VisualLoggerSettings.ArrowLength * Direction,
+			VisualLoggerSettings.Color.ToFColor(true),
 			TEXT(""));
 
 		UE_LEARNING_AGENTS_VLOG_TRANSFORM(VisualLoggerObject, LogLearning, Display,
 			RelativeTransform.GetTranslation(),
 			RelativeTransform.GetRotation(),
-			VisualLoggerColor.ToFColor(true),
+			VisualLoggerSettings.Color.ToFColor(true),
 			TEXT(""));
 
-		UE_LEARNING_AGENTS_VLOG_STRING(VisualLoggerObject, LogLearning, Display, VisualLoggerLocation,
-			VisualLoggerColor.ToFColor(true),
+		UE_LEARNING_AGENTS_VLOG_STRING(VisualLoggerObject, LogLearning, Display, VisualLoggerSettings.Location,
+			VisualLoggerSettings.Color.ToFColor(true),
 			TEXT("Listener: %s\nTag: %s\nAgent Id: % 3i\nLocal Direction: [% 6.1f % 6.1f % 6.1f]\nDirection: [% 6.1f % 6.1f % 6.1f]"),
-			*VisualLoggerListener->GetName(),
+			*VisualLoggerSettings.Listener->GetName(),
 			*Tag.ToString(),
-			VisualLoggerAgentId,
+			VisualLoggerSettings.AgentId,
 			LocalDirection.X, LocalDirection.Y, LocalDirection.Z,
 			Direction.X, Direction.Y, Direction.Z);
 	}
@@ -3318,7 +3263,7 @@ FLearningAgentsActionObjectElement ULearningAgentsActions::MakeDirectionAction(
 		(float)LocalDirection.X,
 		(float)LocalDirection.Y,
 		(float)LocalDirection.Z,
-		}, bActive, Tag);
+		}, bActive, Tag, FLearningAgentsVisualLoggerSettings());
 }
 
 FLearningAgentsActionModifierElement ULearningAgentsActions::MakeNullActionModifier(ULearningAgentsActionModifier* Modifier, const FName Tag)
@@ -3371,7 +3316,7 @@ FLearningAgentsActionModifierElement ULearningAgentsActions::MakeExclusiveDiscre
 	{
 		UE_LOG(LogLearning, Error, TEXT("MakeExclusiveDiscreteActionModifierFromArrayView: MaskedIndices contain duplicates."));
 		return FLearningAgentsActionModifierElement();
-	}	
+	}
 
 	return { Modifier->ActionModifier.CreateDiscreteExclusive({ MaskedIndices }, Tag) };
 }
@@ -3415,7 +3360,7 @@ FLearningAgentsActionModifierElement ULearningAgentsActions::MakeInclusiveDiscre
 	{
 		UE_LOG(LogLearning, Error, TEXT("MakeNamedExclusiveDiscreteActionModifierFromArrayView: MaskedIndices contain duplicates."));
 		return FLearningAgentsActionModifierElement();
-	}	
+	}
 
 	return { Modifier->ActionModifier.CreateDiscreteInclusive({ MaskedIndices }, Tag) };
 }
@@ -3432,12 +3377,12 @@ FLearningAgentsActionModifierElement ULearningAgentsActions::MakeNamedInclusiveD
 		UE_LOG(LogLearning, Error, TEXT("MakeNamedInclusiveDiscreteActionModifierFromArrayView: Modifier is nullptr."));
 		return FLearningAgentsActionModifierElement();
 	}
-	
+
 	if (UE::Learning::Agents::Action::Private::ContainsDuplicates(MaskedNames))
 	{
 		UE_LOG(LogLearning, Error, TEXT("MakeNamedExclusiveDiscreteActionModifierFromArrayView: MaskedNames contain duplicates."));
 		return FLearningAgentsActionModifierElement();
-	}	
+	}
 
 	return { Modifier->ActionModifier.CreateNamedDiscreteInclusive({ MaskedNames }, Tag) };
 }
@@ -3805,8 +3750,8 @@ FLearningAgentsActionModifierElement ULearningAgentsActions::MakeLocationActionM
 		bMaskedY,
 		bMaskedZ
 		}, {
-		(float)LocalLocation.X, 
-		(float)LocalLocation.Y, 
+		(float)LocalLocation.X,
+		(float)LocalLocation.Y,
 		(float)LocalLocation.Z }, Tag);
 }
 
@@ -3821,8 +3766,8 @@ FLearningAgentsActionModifierElement ULearningAgentsActions::MakeScaleActionModi
 		bMaskedY,
 		bMaskedZ
 		}, {
-		(float)LocalLogScale.X, 
-		(float)LocalLogScale.Y, 
+		(float)LocalLogScale.X,
+		(float)LocalLogScale.Y,
 		(float)LocalLogScale.Z }, Tag);
 }
 
@@ -3936,16 +3881,12 @@ bool ULearningAgentsActions::GetContinuousActionNum(int32& OutNum, const ULearni
 }
 
 bool ULearningAgentsActions::GetContinuousAction(
-	TArray<float>& OutValues, 
-	const ULearningAgentsActionObject* Object, 
-	const FLearningAgentsActionObjectElement Element, 
+	TArray<float>& OutValues,
+	const ULearningAgentsActionObject* Object,
+	const FLearningAgentsActionObjectElement Element,
 	const bool bActive,
 	const FName Tag,
-	const bool bVisualLoggerEnabled,
-	ULearningAgentsManagerListener* VisualLoggerListener,
-	const int32 VisualLoggerAgentId,
-	const FVector VisualLoggerLocation,
-	const FLinearColor VisualLoggerColor)
+	const FLearningAgentsVisualLoggerSettings& VisualLoggerSettings)
 {
 	if (!bActive)
 	{
@@ -3961,7 +3902,7 @@ bool ULearningAgentsActions::GetContinuousAction(
 
 	OutValues.SetNumUninitialized(OutValueNum);
 
-	if (!GetContinuousActionToArrayView(OutValues, Object, Element, bActive, Tag, bVisualLoggerEnabled, VisualLoggerListener, VisualLoggerAgentId, VisualLoggerLocation, VisualLoggerColor))
+	if (!GetContinuousActionToArrayView(OutValues, Object, Element, bActive, Tag, VisualLoggerSettings))
 	{
 		OutValues.Empty();
 		return false;
@@ -3971,16 +3912,12 @@ bool ULearningAgentsActions::GetContinuousAction(
 }
 
 bool ULearningAgentsActions::GetContinuousActionToArrayView(
-	TArrayView<float> OutValues, 
-	const ULearningAgentsActionObject* Object, 
-	const FLearningAgentsActionObjectElement Element, 
+	TArrayView<float> OutValues,
+	const ULearningAgentsActionObject* Object,
+	const FLearningAgentsActionObjectElement Element,
 	const bool bActive,
 	const FName Tag,
-	const bool bVisualLoggerEnabled,
-	ULearningAgentsManagerListener* VisualLoggerListener,
-	const int32 VisualLoggerAgentId,
-	const FVector VisualLoggerLocation,
-	const FLinearColor VisualLoggerColor)
+	const FLearningAgentsVisualLoggerSettings& VisualLoggerSettings)
 {
 	if (!bActive)
 	{
@@ -4030,16 +3967,16 @@ bool ULearningAgentsActions::GetContinuousActionToArrayView(
 	UE::Learning::Array::Copy<1, float>(OutValues, Values);
 
 #if UE_LEARNING_AGENTS_ENABLE_VISUAL_LOG
-	if (bVisualLoggerEnabled && VisualLoggerListener)
+	if (VisualLoggerSettings.bEnabled && VisualLoggerSettings.Listener)
 	{
-		const ULearningAgentsVisualLoggerObject* VisualLoggerObject = VisualLoggerListener->GetOrAddVisualLoggerObject(Tag);
+		const ULearningAgentsVisualLoggerObject* VisualLoggerObject = VisualLoggerSettings.Listener->GetOrAddVisualLoggerObject(Tag);
 
-		UE_LEARNING_AGENTS_VLOG_STRING(VisualLoggerObject, LogLearning, Display, VisualLoggerLocation,
-			VisualLoggerColor.ToFColor(true),
+		UE_LEARNING_AGENTS_VLOG_STRING(VisualLoggerObject, LogLearning, Display, VisualLoggerSettings.Location,
+			VisualLoggerSettings.Color.ToFColor(true),
 			TEXT("Listener: %s\nTag: %s\nAgent Id: % 3i\nEncoded: %s\nValues: %s"),
-			*VisualLoggerListener->GetName(),
+			*VisualLoggerSettings.Listener->GetName(),
 			*Tag.ToString(),
-			VisualLoggerAgentId,
+			VisualLoggerSettings.AgentId,
 			*UE::Learning::Array::FormatFloat(OutValues),
 			*UE::Learning::Array::FormatFloat(OutValues)); // Encoded is identical to provided values
 	}
@@ -4049,16 +3986,12 @@ bool ULearningAgentsActions::GetContinuousActionToArrayView(
 }
 
 bool ULearningAgentsActions::GetExclusiveDiscreteAction(
-	int32& OutIndex, 
-	const ULearningAgentsActionObject* Object, 
-	const FLearningAgentsActionObjectElement Element, 
-	const bool bActive,	
+	int32& OutIndex,
+	const ULearningAgentsActionObject* Object,
+	const FLearningAgentsActionObjectElement Element,
+	const bool bActive,
 	const FName Tag,
-	const bool bVisualLoggerEnabled,
-	ULearningAgentsManagerListener* VisualLoggerListener,
-	const int32 VisualLoggerAgentId,
-	const FVector VisualLoggerLocation,
-	const FLinearColor VisualLoggerColor)
+	const FLearningAgentsVisualLoggerSettings& VisualLoggerSettings)
 {
 	if (!bActive)
 	{
@@ -4097,16 +4030,16 @@ bool ULearningAgentsActions::GetExclusiveDiscreteAction(
 	OutIndex = Object->ActionObject.GetDiscreteExclusive(Element.ObjectElement).DiscreteIndex;
 
 #if UE_LEARNING_AGENTS_ENABLE_VISUAL_LOG
-	if (bVisualLoggerEnabled && VisualLoggerListener)
+	if (VisualLoggerSettings.bEnabled && VisualLoggerSettings.Listener)
 	{
-		const ULearningAgentsVisualLoggerObject* VisualLoggerObject = VisualLoggerListener->GetOrAddVisualLoggerObject(Tag);
+		const ULearningAgentsVisualLoggerObject* VisualLoggerObject = VisualLoggerSettings.Listener->GetOrAddVisualLoggerObject(Tag);
 
-		UE_LEARNING_AGENTS_VLOG_STRING(VisualLoggerObject, LogLearning, Display, VisualLoggerLocation,
-			VisualLoggerColor.ToFColor(true),
+		UE_LEARNING_AGENTS_VLOG_STRING(VisualLoggerObject, LogLearning, Display, VisualLoggerSettings.Location,
+			VisualLoggerSettings.Color.ToFColor(true),
 			TEXT("Listener: %s\nTag: %s\nAgent Id: % 3i\nIndex: [%i]"),
-			*VisualLoggerListener->GetName(),
+			*VisualLoggerSettings.Listener->GetName(),
 			*Tag.ToString(),
-			VisualLoggerAgentId,
+			VisualLoggerSettings.AgentId,
 			OutIndex);
 	}
 #endif
@@ -4120,11 +4053,7 @@ bool ULearningAgentsActions::GetNamedExclusiveDiscreteAction(
 	const FLearningAgentsActionObjectElement Element,
 	const bool bActive,
 	const FName Tag,
-	const bool bVisualLoggerEnabled,
-	ULearningAgentsManagerListener* VisualLoggerListener,
-	const int32 VisualLoggerAgentId,
-	const FVector VisualLoggerLocation,
-	const FLinearColor VisualLoggerColor)
+	const FLearningAgentsVisualLoggerSettings& VisualLoggerSettings)
 {
 	if (!bActive)
 	{
@@ -4163,16 +4092,16 @@ bool ULearningAgentsActions::GetNamedExclusiveDiscreteAction(
 	OutName = Object->ActionObject.GetNamedDiscreteExclusive(Element.ObjectElement).ElementName;
 
 #if UE_LEARNING_AGENTS_ENABLE_VISUAL_LOG
-	if (bVisualLoggerEnabled && VisualLoggerListener)
+	if (VisualLoggerSettings.bEnabled && VisualLoggerSettings.Listener)
 	{
-		const ULearningAgentsVisualLoggerObject* VisualLoggerObject = VisualLoggerListener->GetOrAddVisualLoggerObject(Tag);
+		const ULearningAgentsVisualLoggerObject* VisualLoggerObject = VisualLoggerSettings.Listener->GetOrAddVisualLoggerObject(Tag);
 
-		UE_LEARNING_AGENTS_VLOG_STRING(VisualLoggerObject, LogLearning, Display, VisualLoggerLocation,
-			VisualLoggerColor.ToFColor(true),
+		UE_LEARNING_AGENTS_VLOG_STRING(VisualLoggerObject, LogLearning, Display, VisualLoggerSettings.Location,
+			VisualLoggerSettings.Color.ToFColor(true),
 			TEXT("Listener: %s\nTag: %s\nAgent Id: % 3i\nName: [%s]"),
-			*VisualLoggerListener->GetName(),
+			*VisualLoggerSettings.Listener->GetName(),
 			*Tag.ToString(),
-			VisualLoggerAgentId,
+			VisualLoggerSettings.AgentId,
 			*OutName.ToString());
 	}
 #endif
@@ -4221,16 +4150,12 @@ bool ULearningAgentsActions::GetInclusiveDiscreteActionNum(int32& OutNum, const
 }
 
 bool ULearningAgentsActions::GetInclusiveDiscreteAction(
-	TArray<int32>& OutIndices, 
-	const ULearningAgentsActionObject* Object, 
-	const FLearningAgentsActionObjectElement Element, 
+	TArray<int32>& OutIndices,
+	const ULearningAgentsActionObject* Object,
+	const FLearningAgentsActionObjectElement Element,
 	const bool bActive,
 	const FName Tag,
-	const bool bVisualLoggerEnabled,
-	ULearningAgentsManagerListener* VisualLoggerListener,
-	const int32 VisualLoggerAgentId,
-	const FVector VisualLoggerLocation,
-	const FLinearColor VisualLoggerColor)
+	const FLearningAgentsVisualLoggerSettings& VisualLoggerSettings)
 {
 	if (!bActive)
 	{
@@ -4246,7 +4171,7 @@ bool ULearningAgentsActions::GetInclusiveDiscreteAction(
 
 	OutIndices.SetNumUninitialized(OutIndexNum);
 
-	if (!GetInclusiveDiscreteActionToArrayView(OutIndices, Object, Element, bActive, Tag, bVisualLoggerEnabled, VisualLoggerListener, VisualLoggerAgentId, VisualLoggerLocation, VisualLoggerColor))
+	if (!GetInclusiveDiscreteActionToArrayView(OutIndices, Object, Element, bActive, Tag, VisualLoggerSettings))
 	{
 		OutIndices.Empty();
 		return false;
@@ -4257,15 +4182,11 @@ bool ULearningAgentsActions::GetInclusiveDiscreteAction(
 
 bool ULearningAgentsActions::GetInclusiveDiscreteActionToArrayView(
 	TArrayView<int32> OutIndices,
-	const ULearningAgentsActionObject* Object, 
-	const FLearningAgentsActionObjectElement Element, 
+	const ULearningAgentsActionObject* Object,
+	const FLearningAgentsActionObjectElement Element,
 	const bool bActive,
 	const FName Tag,
-	const bool bVisualLoggerEnabled,
-	ULearningAgentsManagerListener* VisualLoggerListener,
-	const int32 VisualLoggerAgentId,
-	const FVector VisualLoggerLocation,
-	const FLinearColor VisualLoggerColor)
+	const FLearningAgentsVisualLoggerSettings& VisualLoggerSettings)
 {
 	if (!bActive)
 	{
@@ -4315,16 +4236,16 @@ bool ULearningAgentsActions::GetInclusiveDiscreteActionToArrayView(
 	UE::Learning::Array::Copy<1, int32>(OutIndices, Indices);
 
 #if UE_LEARNING_AGENTS_ENABLE_VISUAL_LOG
-	if (bVisualLoggerEnabled && VisualLoggerListener)
+	if (VisualLoggerSettings.bEnabled && VisualLoggerSettings.Listener)
 	{
-		const ULearningAgentsVisualLoggerObject* VisualLoggerObject = VisualLoggerListener->GetOrAddVisualLoggerObject(Tag);
+		const ULearningAgentsVisualLoggerObject* VisualLoggerObject = VisualLoggerSettings.Listener->GetOrAddVisualLoggerObject(Tag);
 
-		UE_LEARNING_AGENTS_VLOG_STRING(VisualLoggerObject, LogLearning, Display, VisualLoggerLocation,
-			VisualLoggerColor.ToFColor(true),
+		UE_LEARNING_AGENTS_VLOG_STRING(VisualLoggerObject, LogLearning, Display, VisualLoggerSettings.Location,
+			VisualLoggerSettings.Color.ToFColor(true),
 			TEXT("Listener: %s\nTag: %s\nAgent Id: % 3i\nIndices: %s"),
-			*VisualLoggerListener->GetName(),
+			*VisualLoggerSettings.Listener->GetName(),
 			*Tag.ToString(),
-			VisualLoggerAgentId,
+			VisualLoggerSettings.AgentId,
 			*UE::Learning::Array::FormatInt32(OutIndices, 256));
 	}
 #endif
@@ -4378,11 +4299,7 @@ bool ULearningAgentsActions::GetNamedInclusiveDiscreteAction(
 	const FLearningAgentsActionObjectElement Element,
 	const bool bActive,
 	const FName Tag,
-	const bool bVisualLoggerEnabled,
-	ULearningAgentsManagerListener* VisualLoggerListener,
-	const int32 VisualLoggerAgentId,
-	const FVector VisualLoggerLocation,
-	const FLinearColor VisualLoggerColor)
+	const FLearningAgentsVisualLoggerSettings& VisualLoggerSettings)
 {
 	if (!bActive)
 	{
@@ -4398,7 +4315,7 @@ bool ULearningAgentsActions::GetNamedInclusiveDiscreteAction(
 
 	OutNames.SetNumUninitialized(OutNameNum);
 
-	if (!GetNamedInclusiveDiscreteActionToArrayView(OutNames, Object, Element, bActive, Tag, bVisualLoggerEnabled, VisualLoggerListener, VisualLoggerAgentId, VisualLoggerLocation, VisualLoggerColor))
+	if (!GetNamedInclusiveDiscreteActionToArrayView(OutNames, Object, Element, bActive, Tag, VisualLoggerSettings))
 	{
 		OutNames.Empty();
 		return false;
@@ -4411,13 +4328,9 @@ bool ULearningAgentsActions::GetNamedInclusiveDiscreteActionToArrayView(
 	TArrayView<FName> OutNames,
 	const ULearningAgentsActionObject* Object,
 	const FLearningAgentsActionObjectElement Element,
-	const bool bActive,	
+	const bool bActive,
 	const FName Tag,
-	const bool bVisualLoggerEnabled,
-	ULearningAgentsManagerListener* VisualLoggerListener,
-	const int32 VisualLoggerAgentId,
-	const FVector VisualLoggerLocation,
-	const FLinearColor VisualLoggerColor)
+	const FLearningAgentsVisualLoggerSettings& VisualLoggerSettings)
 {
 	if (!bActive)
 	{
@@ -4467,16 +4380,16 @@ bool ULearningAgentsActions::GetNamedInclusiveDiscreteActionToArrayView(
 	UE::Learning::Array::Copy<1, FName>(OutNames, Names);
 
 #if UE_LEARNING_AGENTS_ENABLE_VISUAL_LOG
-	if (bVisualLoggerEnabled && VisualLoggerListener)
+	if (VisualLoggerSettings.bEnabled && VisualLoggerSettings.Listener)
 	{
-		const ULearningAgentsVisualLoggerObject* VisualLoggerObject = VisualLoggerListener->GetOrAddVisualLoggerObject(Tag);
+		const ULearningAgentsVisualLoggerObject* VisualLoggerObject = VisualLoggerSettings.Listener->GetOrAddVisualLoggerObject(Tag);
 
-		UE_LEARNING_AGENTS_VLOG_STRING(VisualLoggerObject, LogLearning, Display, VisualLoggerLocation,
-			VisualLoggerColor.ToFColor(true),
+		UE_LEARNING_AGENTS_VLOG_STRING(VisualLoggerObject, LogLearning, Display, VisualLoggerSettings.Location,
+			VisualLoggerSettings.Color.ToFColor(true),
 			TEXT("Listener: %s\nTag: %s\nAgent Id: % 3i\nNames: %s"),
-			*VisualLoggerListener->GetName(),
+			*VisualLoggerSettings.Listener->GetName(),
 			*Tag.ToString(),
-			VisualLoggerAgentId,
+			VisualLoggerSettings.AgentId,
 			*UE::Learning::Agents::Action::Private::FormatNames(OutNames, 256));
 	}
 #endif
@@ -4520,7 +4433,7 @@ bool ULearningAgentsActions::GetStructActionNum(int32& OutNum, const ULearningAg
 		OutNum = 0;
 		return false;
 	}
-	
+
 	const UE::Learning::Action::FObjectAndParameters Parameters = Object->ActionObject.GetAnd(Element.ObjectElement);
 
 	OutNum = Parameters.Elements.Num();
@@ -5024,7 +4937,7 @@ bool ULearningAgentsActions::GetStaticArrayActionToArrayView(TArrayView<FLearnin
 
 	if (Object->ActionObject.GetType(Element.ObjectElement) != UE::Learning::Action::EType::Array)
 	{
-		UE_LOG(LogLearning, Error, TEXT("GetStaticArrayActionToArrayView: Action '%s' type does not match. Action is '%s' but asked for '%s'."), 
+		UE_LOG(LogLearning, Error, TEXT("GetStaticArrayActionToArrayView: Action '%s' type does not match. Action is '%s' but asked for '%s'."),
 			*Object->ActionObject.GetTag(Element.ObjectElement).ToString(),
 			UE::Learning::Agents::Action::Private::GetActionTypeString(Object->ActionObject.GetType(Element.ObjectElement)),
 			UE::Learning::Agents::Action::Private::GetActionTypeString(UE::Learning::Action::EType::Array));
@@ -5085,17 +4998,13 @@ bool ULearningAgentsActions::GetPairAction(FLearningAgentsActionObjectElement& O
 }
 
 bool ULearningAgentsActions::GetEnumAction(
-	uint8& OutEnumValue, 
-	const ULearningAgentsActionObject* Object, 
-	const FLearningAgentsActionObjectElement Element, 
-	const UEnum* Enum, 
+	uint8& OutEnumValue,
+	const ULearningAgentsActionObject* Object,
+	const FLearningAgentsActionObjectElement Element,
+	const UEnum* Enum,
 	const bool bActive,
 	const FName Tag,
-	const bool bVisualLoggerEnabled,
-	ULearningAgentsManagerListener* VisualLoggerListener,
-	const int32 VisualLoggerAgentId,
-	const FVector VisualLoggerLocation,
-	const FLinearColor VisualLoggerColor)
+	const FLearningAgentsVisualLoggerSettings& VisualLoggerSettings)
 {
 	if (!bActive)
 	{
@@ -5135,16 +5044,16 @@ bool ULearningAgentsActions::GetEnumAction(
 	OutEnumValue = (uint8)EnumValue;
 
 #if UE_LEARNING_AGENTS_ENABLE_VISUAL_LOG
-	if (bVisualLoggerEnabled && VisualLoggerListener)
+	if (VisualLoggerSettings.bEnabled && VisualLoggerSettings.Listener)
 	{
-		const ULearningAgentsVisualLoggerObject* VisualLoggerObject = VisualLoggerListener->GetOrAddVisualLoggerObject(Tag);
+		const ULearningAgentsVisualLoggerObject* VisualLoggerObject = VisualLoggerSettings.Listener->GetOrAddVisualLoggerObject(Tag);
 
-		UE_LEARNING_AGENTS_VLOG_STRING(VisualLoggerObject, LogLearning, Display, VisualLoggerLocation,
-			VisualLoggerColor.ToFColor(true),
+		UE_LEARNING_AGENTS_VLOG_STRING(VisualLoggerObject, LogLearning, Display, VisualLoggerSettings.Location,
+			VisualLoggerSettings.Color.ToFColor(true),
 			TEXT("Listener: %s\nTag: %s\nAgent Id: % 3i\nEnum: %s\nSize: [%i]\nValue: [%s]\nIndex: [%i]"),
-			*VisualLoggerListener->GetName(),
+			*VisualLoggerSettings.Listener->GetName(),
 			*Tag.ToString(),
-			VisualLoggerAgentId,
+			VisualLoggerSettings.AgentId,
 			*Enum->GetName(),
 			Enum->NumEnums() - 1,
 			*Enum->GetDisplayNameTextByValue(OutEnumValue).ToString(),
@@ -5156,17 +5065,13 @@ bool ULearningAgentsActions::GetEnumAction(
 }
 
 bool ULearningAgentsActions::GetBitmaskAction(
-	int32& OutBitmaskValue, 
-	const ULearningAgentsActionObject* Object, 
-	const FLearningAgentsActionObjectElement Element, 
-	const UEnum* Enum, 
+	int32& OutBitmaskValue,
+	const ULearningAgentsActionObject* Object,
+	const FLearningAgentsActionObjectElement Element,
+	const UEnum* Enum,
 	const bool bActive,
 	const FName Tag,
-	const bool bVisualLoggerEnabled,
-	ULearningAgentsManagerListener* VisualLoggerListener,
-	const int32 VisualLoggerAgentId,
-	const FVector VisualLoggerLocation,
-	const FLinearColor VisualLoggerColor)
+	const FLearningAgentsVisualLoggerSettings& VisualLoggerSettings)
 {
 	if (!bActive)
 	{
@@ -5216,9 +5121,9 @@ bool ULearningAgentsActions::GetBitmaskAction(
 	}
 
 #if UE_LEARNING_AGENTS_ENABLE_VISUAL_LOG
-	if (bVisualLoggerEnabled && VisualLoggerListener)
+	if (VisualLoggerSettings.bEnabled && VisualLoggerSettings.Listener)
 	{
-		const ULearningAgentsVisualLoggerObject* VisualLoggerObject = VisualLoggerListener->GetOrAddVisualLoggerObject(Tag);
+		const ULearningAgentsVisualLoggerObject* VisualLoggerObject = VisualLoggerSettings.Listener->GetOrAddVisualLoggerObject(Tag);
 
 		FString ValuesString;
 		FString IndicesString;
@@ -5235,12 +5140,12 @@ bool ULearningAgentsActions::GetBitmaskAction(
 		ValuesString = ValuesString.TrimEnd();
 		IndicesString = IndicesString.TrimEnd();
 
-		UE_LEARNING_AGENTS_VLOG_STRING(VisualLoggerObject, LogLearning, Display, VisualLoggerLocation,
-			VisualLoggerColor.ToFColor(true),
+		UE_LEARNING_AGENTS_VLOG_STRING(VisualLoggerObject, LogLearning, Display, VisualLoggerSettings.Location,
+			VisualLoggerSettings.Color.ToFColor(true),
 			TEXT("Listener: %s\nTag: %s\nAgent Id: % 3i\nEnum: %s\nSize: [%i]\nValues: [%s]\nIndices: [%s]"),
-			*VisualLoggerListener->GetName(),
+			*VisualLoggerSettings.Listener->GetName(),
 			*Tag.ToString(),
-			VisualLoggerAgentId,
+			VisualLoggerSettings.AgentId,
 			*Enum->GetName(),
 			Enum->NumEnums() - 1,
 			*ValuesString,
@@ -5329,16 +5234,12 @@ bool ULearningAgentsActions::GetEncodingAction(FLearningAgentsActionObjectElemen
 }
 
 bool ULearningAgentsActions::GetBoolAction(
-	bool& bOutValue, 
-	const ULearningAgentsActionObject* Object, 
-	const FLearningAgentsActionObjectElement Element, 
-	const bool bActive,	
+	bool& bOutValue,
+	const ULearningAgentsActionObject* Object,
+	const FLearningAgentsActionObjectElement Element,
+	const bool bActive,
 	const FName Tag,
-	const bool bVisualLoggerEnabled,
-	ULearningAgentsManagerListener* VisualLoggerListener,
-	const int32 VisualLoggerAgentId,
-	const FVector VisualLoggerLocation,
-	const FLinearColor VisualLoggerColor)
+	const FLearningAgentsVisualLoggerSettings& VisualLoggerSettings)
 {
 	if (!bActive)
 	{
@@ -5353,36 +5254,32 @@ bool ULearningAgentsActions::GetBoolAction(
 	}
 
 	bOutValue = OutIndex == 1;
-	
+
 #if UE_LEARNING_AGENTS_ENABLE_VISUAL_LOG
-	if (bVisualLoggerEnabled && VisualLoggerListener)
+	if (VisualLoggerSettings.bEnabled && VisualLoggerSettings.Listener)
 	{
-		const ULearningAgentsVisualLoggerObject* VisualLoggerObject = VisualLoggerListener->GetOrAddVisualLoggerObject(Tag);
+		const ULearningAgentsVisualLoggerObject* VisualLoggerObject = VisualLoggerSettings.Listener->GetOrAddVisualLoggerObject(Tag);
 
-		UE_LEARNING_AGENTS_VLOG_STRING(VisualLoggerObject, LogLearning, Display, VisualLoggerLocation,
-			VisualLoggerColor.ToFColor(true),
+		UE_LEARNING_AGENTS_VLOG_STRING(VisualLoggerObject, LogLearning, Display, VisualLoggerSettings.Location,
+			VisualLoggerSettings.Color.ToFColor(true),
 			TEXT("Listener: %s\nTag: %s\nAgent Id: % 3i\nValue: [%s]"),
-			*VisualLoggerListener->GetName(),
+			*VisualLoggerSettings.Listener->GetName(),
 			*Tag.ToString(),
-			VisualLoggerAgentId,
+			VisualLoggerSettings.AgentId,
 			bOutValue ? TEXT("true") : TEXT("false"));
 	}
 #endif
-	
+
 	return true;
 }
 
 bool ULearningAgentsActions::GetFloatAction(
-	float& OutValue, 
-	const ULearningAgentsActionObject* Object, 
-	const FLearningAgentsActionObjectElement Element, 
-	const bool bActive, 
+	float& OutValue,
+	const ULearningAgentsActionObject* Object,
+	const FLearningAgentsActionObjectElement Element,
+	const bool bActive,
 	const FName Tag,
-	const bool bVisualLoggerEnabled,
-	ULearningAgentsManagerListener* VisualLoggerListener,
-	const int32 VisualLoggerAgentId,
-	const FVector VisualLoggerLocation,
-	const FLinearColor VisualLoggerColor)
+	const FLearningAgentsVisualLoggerSettings& VisualLoggerSettings)
 {
 	if (!bActive)
 	{
@@ -5396,16 +5293,16 @@ bool ULearningAgentsActions::GetFloatAction(
 	}
 
 #if UE_LEARNING_AGENTS_ENABLE_VISUAL_LOG
-	if (bVisualLoggerEnabled && VisualLoggerListener)
+	if (VisualLoggerSettings.bEnabled && VisualLoggerSettings.Listener)
 	{
-		const ULearningAgentsVisualLoggerObject* VisualLoggerObject = VisualLoggerListener->GetOrAddVisualLoggerObject(Tag);
+		const ULearningAgentsVisualLoggerObject* VisualLoggerObject = VisualLoggerSettings.Listener->GetOrAddVisualLoggerObject(Tag);
 
-		UE_LEARNING_AGENTS_VLOG_STRING(VisualLoggerObject, LogLearning, Display, VisualLoggerLocation,
-			VisualLoggerColor.ToFColor(true),
+		UE_LEARNING_AGENTS_VLOG_STRING(VisualLoggerObject, LogLearning, Display, VisualLoggerSettings.Location,
+			VisualLoggerSettings.Color.ToFColor(true),
 			TEXT("Listener: %s\nTag: %s\nAgent Id: % 3i\nValue: [% 6.1f]"),
-			*VisualLoggerListener->GetName(),
+			*VisualLoggerSettings.Listener->GetName(),
 			*Tag.ToString(),
-			VisualLoggerAgentId,
+			VisualLoggerSettings.AgentId,
 			OutValue);
 	}
 #endif
@@ -5414,17 +5311,13 @@ bool ULearningAgentsActions::GetFloatAction(
 }
 
 bool ULearningAgentsActions::GetLocationAction(
-	FVector& OutLocation, 
-	const ULearningAgentsActionObject* Object, 
-	const FLearningAgentsActionObjectElement Element, 
-	const FTransform RelativeTransform, 
+	FVector& OutLocation,
+	const ULearningAgentsActionObject* Object,
+	const FLearningAgentsActionObjectElement Element,
+	const FTransform RelativeTransform,
 	const bool bActive,
 	const FName Tag,
-	const bool bVisualLoggerEnabled,
-	ULearningAgentsManagerListener* VisualLoggerListener,
-	const int32 VisualLoggerAgentId,
-	const FVector VisualLoggerLocation,
-	const FLinearColor VisualLoggerColor)
+	const FLearningAgentsVisualLoggerSettings& VisualLoggerSettings)
 {
 	if (!bActive)
 	{
@@ -5440,57 +5333,52 @@ bool ULearningAgentsActions::GetLocationAction(
 
 	const FVector LocalLocation = FVector(OutValues[0], OutValues[1], OutValues[2]);
 	OutLocation = RelativeTransform.TransformPosition(LocalLocation);
-	
+
 #if UE_LEARNING_AGENTS_ENABLE_VISUAL_LOG
-	if (bVisualLoggerEnabled && VisualLoggerListener)
+	if (VisualLoggerSettings.bEnabled && VisualLoggerSettings.Listener)
 	{
-		const ULearningAgentsVisualLoggerObject* VisualLoggerObject = VisualLoggerListener->GetOrAddVisualLoggerObject(Tag);
+		const ULearningAgentsVisualLoggerObject* VisualLoggerObject = VisualLoggerSettings.Listener->GetOrAddVisualLoggerObject(Tag);
 
 		UE_LEARNING_AGENTS_VLOG_LOCATION(VisualLoggerObject, LogLearning, Display,
 			OutLocation,
 			10,
-			VisualLoggerColor.ToFColor(true),
+			VisualLoggerSettings.Color.ToFColor(true),
 			TEXT(""));
 
 		UE_LEARNING_AGENTS_VLOG_SEGMENT(VisualLoggerObject, LogLearning, Display,
 			RelativeTransform.GetTranslation(),
 			OutLocation,
-			VisualLoggerColor.ToFColor(true),
+			VisualLoggerSettings.Color.ToFColor(true),
 			TEXT(""));
 
 		UE_LEARNING_AGENTS_VLOG_TRANSFORM(VisualLoggerObject, LogLearning, Display,
 			RelativeTransform.GetTranslation(),
 			RelativeTransform.GetRotation(),
-			VisualLoggerColor.ToFColor(true),
+			VisualLoggerSettings.Color.ToFColor(true),
 			TEXT(""));
 
-		UE_LEARNING_AGENTS_VLOG_STRING(VisualLoggerObject, LogLearning, Display, VisualLoggerLocation,
-			VisualLoggerColor.ToFColor(true),
+		UE_LEARNING_AGENTS_VLOG_STRING(VisualLoggerObject, LogLearning, Display, VisualLoggerSettings.Location,
+			VisualLoggerSettings.Color.ToFColor(true),
 			TEXT("Listener: %s\nTag: %s\nAgent Id: % 3i\nLocal Location: [% 6.1f % 6.1f % 6.1f]\nLocation: [% 6.1f % 6.1f % 6.1f]"),
-			*VisualLoggerListener->GetName(),
+			*VisualLoggerSettings.Listener->GetName(),
 			*Tag.ToString(),
-			VisualLoggerAgentId,
+			VisualLoggerSettings.AgentId,
 			LocalLocation.X, LocalLocation.Y, LocalLocation.Z,
 			OutLocation.X, OutLocation.Y, OutLocation.Z);
 	}
 #endif
-	
+
 	return true;
 }
 
 bool ULearningAgentsActions::GetRotationAction(
-	FRotator& OutRotation, 
-	const ULearningAgentsActionObject* Object, 
-	const FLearningAgentsActionObjectElement Element, 
-	const FRotator RelativeRotation, 
+	FRotator& OutRotation,
+	const ULearningAgentsActionObject* Object,
+	const FLearningAgentsActionObjectElement Element,
+	const FRotator RelativeRotation,
 	const bool bActive,
 	const FName Tag,
-	const bool bVisualLoggerEnabled,
-	ULearningAgentsManagerListener* VisualLoggerListener,
-	const int32 VisualLoggerAgentId,
-	const FVector VisualLoggerRotationLocation,
-	const FVector VisualLoggerLocation,
-	const FLinearColor VisualLoggerColor)
+	const FLearningAgentsVisualLoggerSettings& VisualLoggerSettings)
 {
 	if (!bActive)
 	{
@@ -5498,7 +5386,7 @@ bool ULearningAgentsActions::GetRotationAction(
 	}
 
 	FQuat OutRotationQuat;
-	if (!GetRotationActionAsQuat(OutRotationQuat, Object, Element, FQuat::MakeFromRotator(RelativeRotation), bActive, Tag, bVisualLoggerEnabled, VisualLoggerListener, VisualLoggerAgentId, VisualLoggerRotationLocation, VisualLoggerLocation, VisualLoggerColor))
+	if (!GetRotationActionAsQuat(OutRotationQuat, Object, Element, FQuat::MakeFromRotator(RelativeRotation), bActive, Tag, VisualLoggerSettings))
 	{
 		OutRotation = FRotator::ZeroRotator;
 		return false;
@@ -5509,18 +5397,13 @@ bool ULearningAgentsActions::GetRotationAction(
 }
 
 bool ULearningAgentsActions::GetRotationActionAsQuat(
-	FQuat& OutRotation, 
-	const ULearningAgentsActionObject* Object, 
-	const FLearningAgentsActionObjectElement Element, 
-	const FQuat RelativeRotation, 
+	FQuat& OutRotation,
+	const ULearningAgentsActionObject* Object,
+	const FLearningAgentsActionObjectElement Element,
+	const FQuat RelativeRotation,
 	const bool bActive,
 	const FName Tag,
-	const bool bVisualLoggerEnabled,
-	ULearningAgentsManagerListener* VisualLoggerListener,
-	const int32 VisualLoggerAgentId,
-	const FVector VisualLoggerRotationLocation,
-	const FVector VisualLoggerLocation,
-	const FLinearColor VisualLoggerColor)
+	const FLearningAgentsVisualLoggerSettings& VisualLoggerSettings)
 {
 	if (!bActive)
 	{
@@ -5539,22 +5422,22 @@ bool ULearningAgentsActions::GetRotationActionAsQuat(
 	OutRotation = RelativeRotation * LocalRotation;
 
 #if UE_LEARNING_AGENTS_ENABLE_VISUAL_LOG
-	if (bVisualLoggerEnabled && VisualLoggerListener)
+	if (VisualLoggerSettings.bEnabled && VisualLoggerSettings.Listener)
 	{
-		const ULearningAgentsVisualLoggerObject* VisualLoggerObject = VisualLoggerListener->GetOrAddVisualLoggerObject(Tag);
+		const ULearningAgentsVisualLoggerObject* VisualLoggerObject = VisualLoggerSettings.Listener->GetOrAddVisualLoggerObject(Tag);
 
 		UE_LEARNING_AGENTS_VLOG_TRANSFORM(VisualLoggerObject, LogLearning, Display,
-			VisualLoggerRotationLocation,
+			VisualLoggerSettings.DebugLocation,
 			LocalRotation.Rotator(),
-			VisualLoggerColor.ToFColor(true),
+			VisualLoggerSettings.Color.ToFColor(true),
 			TEXT(""));
 
-		UE_LEARNING_AGENTS_VLOG_STRING(VisualLoggerObject, LogLearning, Display, VisualLoggerLocation,
-			VisualLoggerColor.ToFColor(true),
+		UE_LEARNING_AGENTS_VLOG_STRING(VisualLoggerObject, LogLearning, Display, VisualLoggerSettings.Location,
+			VisualLoggerSettings.Color.ToFColor(true),
 			TEXT("Listener: %s\nTag: %s\nAgent Id: % 3i\nLocal Rotation Vector: [% 6.1f % 6.1f % 6.1f]\nLocal Rotation: [% 6.1f % 6.1f % 6.1f % 6.1f]\nRotation: [% 6.1f % 6.1f % 6.1f % 6.1f]"),
-			*VisualLoggerListener->GetName(),
+			*VisualLoggerSettings.Listener->GetName(),
 			*Tag.ToString(),
-			VisualLoggerAgentId,
+			VisualLoggerSettings.AgentId,
 			LocalRotationVector.X, LocalRotationVector.Y, LocalRotationVector.Z,
 			LocalRotation.X, LocalRotation.Y, LocalRotation.Z, LocalRotation.W,
 			OutRotation.X, OutRotation.Y, OutRotation.Z, OutRotation.W);
@@ -5565,17 +5448,13 @@ bool ULearningAgentsActions::GetRotationActionAsQuat(
 }
 
 bool ULearningAgentsActions::GetScaleAction(
-	FVector& OutScale, 
-	const ULearningAgentsActionObject* Object, 
-	const FLearningAgentsActionObjectElement Element, 
-	const FVector RelativeScale, 
+	FVector& OutScale,
+	const ULearningAgentsActionObject* Object,
+	const FLearningAgentsActionObjectElement Element,
+	const FVector RelativeScale,
 	const bool bActive,
 	const FName Tag,
-	const bool bVisualLoggerEnabled,
-	ULearningAgentsManagerListener* VisualLoggerListener,
-	const int32 VisualLoggerAgentId,
-	const FVector VisualLoggerLocation,
-	const FLinearColor VisualLoggerColor)
+	const FLearningAgentsVisualLoggerSettings& VisualLoggerSettings)
 {
 	if (!bActive)
 	{
@@ -5593,16 +5472,16 @@ bool ULearningAgentsActions::GetScaleAction(
 	OutScale = RelativeScale * LocalScaleVector;
 
 #if UE_LEARNING_AGENTS_ENABLE_VISUAL_LOG
-	if (bVisualLoggerEnabled && VisualLoggerListener)
+	if (VisualLoggerSettings.bEnabled && VisualLoggerSettings.Listener)
 	{
-		const ULearningAgentsVisualLoggerObject* VisualLoggerObject = VisualLoggerListener->GetOrAddVisualLoggerObject(Tag);
+		const ULearningAgentsVisualLoggerObject* VisualLoggerObject = VisualLoggerSettings.Listener->GetOrAddVisualLoggerObject(Tag);
 
-		UE_LEARNING_AGENTS_VLOG_STRING(VisualLoggerObject, LogLearning, Display, VisualLoggerLocation,
-			VisualLoggerColor.ToFColor(true),
+		UE_LEARNING_AGENTS_VLOG_STRING(VisualLoggerObject, LogLearning, Display, VisualLoggerSettings.Location,
+			VisualLoggerSettings.Color.ToFColor(true),
 			TEXT("Listener: %s\nTag: %s\nAgent Id: % 3i\nLocal Scale: [% 6.1f % 6.1f % 6.1f]\nScale: [% 6.1f % 6.1f % 6.1f]"),
-			*VisualLoggerListener->GetName(),
+			*VisualLoggerSettings.Listener->GetName(),
 			*Tag.ToString(),
-			VisualLoggerAgentId,
+			VisualLoggerSettings.AgentId,
 			LocalScaleVector.X, LocalScaleVector.Y, LocalScaleVector.Z,
 			OutScale.X, OutScale.Y, OutScale.Z);
 	}
@@ -5612,17 +5491,13 @@ bool ULearningAgentsActions::GetScaleAction(
 }
 
 bool ULearningAgentsActions::GetTransformAction(
-	FTransform& OutTransform, 
-	const ULearningAgentsActionObject* Object, 
-	const FLearningAgentsActionObjectElement Element, 
-	const FTransform RelativeTransform, 
+	FTransform& OutTransform,
+	const ULearningAgentsActionObject* Object,
+	const FLearningAgentsActionObjectElement Element,
+	const FTransform RelativeTransform,
 	const bool bActive,
 	const FName Tag,
-	const bool bVisualLoggerEnabled,
-	ULearningAgentsManagerListener* VisualLoggerListener,
-	const int32 VisualLoggerAgentId,
-	const FVector VisualLoggerLocation,
-	const FLinearColor VisualLoggerColor)
+	const FLearningAgentsVisualLoggerSettings& VisualLoggerSettings)
 {
 	if (!bActive)
 	{
@@ -5661,22 +5536,22 @@ bool ULearningAgentsActions::GetTransformAction(
 	OutTransform = FTransform(OutRotation, OutLocation, OutScale);
 
 #if UE_LEARNING_AGENTS_ENABLE_VISUAL_LOG
-	if (bVisualLoggerEnabled && VisualLoggerListener)
+	if (VisualLoggerSettings.bEnabled && VisualLoggerSettings.Listener)
 	{
-		const ULearningAgentsVisualLoggerObject* VisualLoggerObject = VisualLoggerListener->GetOrAddVisualLoggerObject(Tag);
+		const ULearningAgentsVisualLoggerObject* VisualLoggerObject = VisualLoggerSettings.Listener->GetOrAddVisualLoggerObject(Tag);
 
 		UE_LEARNING_AGENTS_VLOG_TRANSFORM(VisualLoggerObject, LogLearning, Display,
 			OutLocation,
 			OutRotation,
-			VisualLoggerColor.ToFColor(true),
+			VisualLoggerSettings.Color.ToFColor(true),
 			TEXT(""));
 
-		UE_LEARNING_AGENTS_VLOG_STRING(VisualLoggerObject, LogLearning, Display, VisualLoggerLocation,
-			VisualLoggerColor.ToFColor(true),
+		UE_LEARNING_AGENTS_VLOG_STRING(VisualLoggerObject, LogLearning, Display, VisualLoggerSettings.Location,
+			VisualLoggerSettings.Color.ToFColor(true),
 			TEXT("Listener: %s\nTag: %s\nAgent Id: % 3i\nLocation: [% 6.1f % 6.1f % 6.1f]\nRotation: [% 6.1f % 6.1f % 6.1f % 6.1f]\nScale: [% 6.1f % 6.1f % 6.1f]"),
-			*VisualLoggerListener->GetName(),
+			*VisualLoggerSettings.Listener->GetName(),
 			*Tag.ToString(),
-			VisualLoggerAgentId,
+			VisualLoggerSettings.AgentId,
 			OutLocation.X, OutLocation.Y, OutLocation.Z,
 			OutRotation.X, OutRotation.Y, OutRotation.Z, OutRotation.W,
 			OutScale.X, OutScale.Y, OutScale.Z);
@@ -5687,18 +5562,13 @@ bool ULearningAgentsActions::GetTransformAction(
 }
 
 bool ULearningAgentsActions::GetAngleAction(
-	float& OutAngle, 
-	const ULearningAgentsActionObject* Object, 
-	const FLearningAgentsActionObjectElement Element, 
-	const float RelativeAngle, 
+	float& OutAngle,
+	const ULearningAgentsActionObject* Object,
+	const FLearningAgentsActionObjectElement Element,
+	const float RelativeAngle,
 	const bool bActive,
 	const FName Tag,
-	const bool bVisualLoggerEnabled,
-	ULearningAgentsManagerListener* VisualLoggerListener,
-	const int32 VisualLoggerAgentId,
-	const FVector VisualLoggerAngleLocation,
-	const FVector VisualLoggerLocation,
-	const FLinearColor VisualLoggerColor)
+	const FLearningAgentsVisualLoggerSettings& VisualLoggerSettings)
 {
 	if (!bActive)
 	{
@@ -5715,24 +5585,24 @@ bool ULearningAgentsActions::GetAngleAction(
 	OutAngle = RelativeAngle + LocalAngle;
 
 #if UE_LEARNING_AGENTS_ENABLE_VISUAL_LOG
-	if (bVisualLoggerEnabled && VisualLoggerListener)
+	if (VisualLoggerSettings.bEnabled && VisualLoggerSettings.Listener)
 	{
-		const ULearningAgentsVisualLoggerObject* VisualLoggerObject = VisualLoggerListener->GetOrAddVisualLoggerObject(Tag);
+		const ULearningAgentsVisualLoggerObject* VisualLoggerObject = VisualLoggerSettings.Listener->GetOrAddVisualLoggerObject(Tag);
 
 		UE_LEARNING_AGENTS_VLOG_ANGLE_DEGREES(VisualLoggerObject, LogLearning, Display,
 			OutAngle,
 			0.0f,
-			VisualLoggerAngleLocation,
+			VisualLoggerSettings.DebugLocation,
 			10.0f,
-			VisualLoggerColor.ToFColor(true),
+			VisualLoggerSettings.Color.ToFColor(true),
 			TEXT(""));
 
-		UE_LEARNING_AGENTS_VLOG_STRING(VisualLoggerObject, LogLearning, Display, VisualLoggerLocation,
-			VisualLoggerColor.ToFColor(true),
+		UE_LEARNING_AGENTS_VLOG_STRING(VisualLoggerObject, LogLearning, Display, VisualLoggerSettings.Location,
+			VisualLoggerSettings.Color.ToFColor(true),
 			TEXT("Listener: %s\nTag: %s\nAgent Id: % 3i\nLocal Angle: [% 6.1f]\nAngle: [% 6.1f]"),
-			*VisualLoggerListener->GetName(),
+			*VisualLoggerSettings.Listener->GetName(),
 			*Tag.ToString(),
-			VisualLoggerAgentId,
+			VisualLoggerSettings.AgentId,
 			LocalAngle,
 			OutAngle);
 	}
@@ -5742,25 +5612,20 @@ bool ULearningAgentsActions::GetAngleAction(
 }
 
 bool ULearningAgentsActions::GetAngleActionRadians(
-	float& OutAngle, 
-	const ULearningAgentsActionObject* Object, 
-	const FLearningAgentsActionObjectElement Element, 
-	const float RelativeAngle, 
-	const bool bActive,	
+	float& OutAngle,
+	const ULearningAgentsActionObject* Object,
+	const FLearningAgentsActionObjectElement Element,
+	const float RelativeAngle,
+	const bool bActive,
 	const FName Tag,
-	const bool bVisualLoggerEnabled,
-	ULearningAgentsManagerListener* VisualLoggerListener,
-	const int32 VisualLoggerAgentId,
-	const FVector VisualLoggerAngleLocation,
-	const FVector VisualLoggerLocation,
-	const FLinearColor VisualLoggerColor)
+	const FLearningAgentsVisualLoggerSettings& VisualLoggerSettings)
 {
 	if (!bActive)
 	{
 		return GetNullAction(Object, Element, Tag);
 	}
 
-	if (!GetAngleAction(OutAngle, Object, Element, FMath::RadiansToDegrees(RelativeAngle), bActive, Tag, bVisualLoggerEnabled, VisualLoggerListener, VisualLoggerAgentId, VisualLoggerAngleLocation, VisualLoggerLocation, VisualLoggerColor))
+	if (!GetAngleAction(OutAngle, Object, Element, FMath::RadiansToDegrees(RelativeAngle), bActive, Tag, VisualLoggerSettings))
 	{
 		OutAngle = 0.0f;
 		return false;
@@ -5771,18 +5636,13 @@ bool ULearningAgentsActions::GetAngleActionRadians(
 }
 
 bool ULearningAgentsActions::GetVelocityAction(
-	FVector& OutVelocity, 
-	const ULearningAgentsActionObject* Object, 
-	const FLearningAgentsActionObjectElement Element, 
-	const FTransform RelativeTransform, 
+	FVector& OutVelocity,
+	const ULearningAgentsActionObject* Object,
+	const FLearningAgentsActionObjectElement Element,
+	const FTransform RelativeTransform,
 	const bool bActive,
 	const FName Tag,
-	const bool bVisualLoggerEnabled,
-	ULearningAgentsManagerListener* VisualLoggerListener,
-	const int32 VisualLoggerAgentId,
-	const FVector VisualLoggerVelocityLocation,
-	const FVector VisualLoggerLocation,
-	const FLinearColor VisualLoggerColor)
+	const FLearningAgentsVisualLoggerSettings& VisualLoggerSettings)
 {
 	if (!bActive)
 	{
@@ -5800,28 +5660,28 @@ bool ULearningAgentsActions::GetVelocityAction(
 	OutVelocity = RelativeTransform.TransformVector(LocalVelocity);
 
 #if UE_LEARNING_AGENTS_ENABLE_VISUAL_LOG
-	if (bVisualLoggerEnabled && VisualLoggerListener)
+	if (VisualLoggerSettings.bEnabled && VisualLoggerSettings.Listener)
 	{
-		const ULearningAgentsVisualLoggerObject* VisualLoggerObject = VisualLoggerListener->GetOrAddVisualLoggerObject(Tag);
+		const ULearningAgentsVisualLoggerObject* VisualLoggerObject = VisualLoggerSettings.Listener->GetOrAddVisualLoggerObject(Tag);
 
 		UE_LEARNING_AGENTS_VLOG_ARROW(VisualLoggerObject, LogLearning, Display,
-			VisualLoggerVelocityLocation,
-			VisualLoggerVelocityLocation + OutVelocity,
-			VisualLoggerColor.ToFColor(true),
+			VisualLoggerSettings.DebugLocation,
+			VisualLoggerSettings.DebugLocation + OutVelocity,
+			VisualLoggerSettings.Color.ToFColor(true),
 			TEXT(""));
 
 		UE_LEARNING_AGENTS_VLOG_TRANSFORM(VisualLoggerObject, LogLearning, Display,
 			RelativeTransform.GetTranslation(),
 			RelativeTransform.GetRotation(),
-			VisualLoggerColor.ToFColor(true),
+			VisualLoggerSettings.Color.ToFColor(true),
 			TEXT(""));
 
-		UE_LEARNING_AGENTS_VLOG_STRING(VisualLoggerObject, LogLearning, Display, VisualLoggerLocation,
-			VisualLoggerColor.ToFColor(true),
+		UE_LEARNING_AGENTS_VLOG_STRING(VisualLoggerObject, LogLearning, Display, VisualLoggerSettings.Location,
+			VisualLoggerSettings.Color.ToFColor(true),
 			TEXT("Listener: %s\nTag: %s\nAgent Id: % 3i\nLocal Velocity: [% 6.1f % 6.1f % 6.1f]\nVelocity: [% 6.1f % 6.1f % 6.1f]"),
-			*VisualLoggerListener->GetName(),
+			*VisualLoggerSettings.Listener->GetName(),
 			*Tag.ToString(),
-			VisualLoggerAgentId,
+			VisualLoggerSettings.AgentId,
 			LocalVelocity.X, LocalVelocity.Y, LocalVelocity.Z,
 			OutVelocity.X, OutVelocity.Y, OutVelocity.Z);
 	}
@@ -5831,19 +5691,13 @@ bool ULearningAgentsActions::GetVelocityAction(
 }
 
 bool ULearningAgentsActions::GetDirectionAction(
-	FVector& OutDirection, 
-	const ULearningAgentsActionObject* Object, 
-	const FLearningAgentsActionObjectElement Element, 
-	const FTransform RelativeTransform, 
-	const bool bActive,	
+	FVector& OutDirection,
+	const ULearningAgentsActionObject* Object,
+	const FLearningAgentsActionObjectElement Element,
+	const FTransform RelativeTransform,
+	const bool bActive,
 	const FName Tag,
-	const bool bVisualLoggerEnabled,
-	ULearningAgentsManagerListener* VisualLoggerListener,
-	const int32 VisualLoggerAgentId,
-	const FVector VisualLoggerDirectionLocation,
-	const FVector VisualLoggerLocation,
-	const float VisualLoggerArrowLength,
-	const FLinearColor VisualLoggerColor)
+	const FLearningAgentsVisualLoggerSettings& VisualLoggerSettings)
 {
 	if (!bActive)
 	{
@@ -5861,28 +5715,28 @@ bool ULearningAgentsActions::GetDirectionAction(
 	OutDirection = RelativeTransform.TransformVectorNoScale(LocalDirection);
 
 #if UE_LEARNING_AGENTS_ENABLE_VISUAL_LOG
-	if (bVisualLoggerEnabled && VisualLoggerListener)
+	if (VisualLoggerSettings.bEnabled && VisualLoggerSettings.Listener)
 	{
-		const ULearningAgentsVisualLoggerObject* VisualLoggerObject = VisualLoggerListener->GetOrAddVisualLoggerObject(Tag);
+		const ULearningAgentsVisualLoggerObject* VisualLoggerObject = VisualLoggerSettings.Listener->GetOrAddVisualLoggerObject(Tag);
 
 		UE_LEARNING_AGENTS_VLOG_ARROW(VisualLoggerObject, LogLearning, Display,
-			VisualLoggerDirectionLocation,
-			VisualLoggerDirectionLocation + VisualLoggerArrowLength * OutDirection,
-			VisualLoggerColor.ToFColor(true),
+			VisualLoggerSettings.DebugLocation,
+			VisualLoggerSettings.DebugLocation + VisualLoggerSettings.ArrowLength * OutDirection,
+			VisualLoggerSettings.Color.ToFColor(true),
 			TEXT(""));
 
 		UE_LEARNING_AGENTS_VLOG_TRANSFORM(VisualLoggerObject, LogLearning, Display,
 			RelativeTransform.GetTranslation(),
 			RelativeTransform.GetRotation(),
-			VisualLoggerColor.ToFColor(true),
+			VisualLoggerSettings.Color.ToFColor(true),
 			TEXT(""));
 
-		UE_LEARNING_AGENTS_VLOG_STRING(VisualLoggerObject, LogLearning, Display, VisualLoggerLocation,
-			VisualLoggerColor.ToFColor(true),
+		UE_LEARNING_AGENTS_VLOG_STRING(VisualLoggerObject, LogLearning, Display, VisualLoggerSettings.Location,
+			VisualLoggerSettings.Color.ToFColor(true),
 			TEXT("Listener: %s\nTag: %s\nAgent Id: % 3i\nEncoded: [% 6.2f % 6.2f % 6.2f]\nLocal Direction: [% 6.1f % 6.1f % 6.1f]\nDirection: [% 6.1f % 6.1f % 6.1f]"),
-			*VisualLoggerListener->GetName(),
+			*VisualLoggerSettings.Listener->GetName(),
 			*Tag.ToString(),
-			VisualLoggerAgentId,
+			VisualLoggerSettings.AgentId,
 			OutValues[0], OutValues[1], OutValues[2],
 			LocalDirection.X, LocalDirection.Y, LocalDirection.Z,
 			OutDirection.X, OutDirection.Y, OutDirection.Z);
@@ -5890,4 +5744,4 @@ bool ULearningAgentsActions::GetDirectionAction(
 #endif
 
 	return true;
-}
+}
\ No newline at end of file
