[LearningAgents] Flow Matching and Normalizations

diff --git a/Engine/Plugins/Experimental/LearningAgents/Source/LearningAgents/Private/LearningAgentsCritic.cpp b/Engine/Plugins/Experimental/LearningAgents/Source/LearningAgents/Private/LearningAgentsCritic.cpp
--- a/Engine/Plugins/Experimental/LearningAgents/Source/LearningAgents/Private/LearningAgentsCritic.cpp
+++ b/Engine/Plugins/Experimental/LearningAgents/Source/LearningAgents/Private/LearningAgentsCritic.cpp
@@ -19,22 +19,6 @@
 
 #include UE_INLINE_GENERATED_CPP_BY_NAME(LearningAgentsCritic)
 
-namespace UE::Learning::Agents::Critic::Private
-{
-	static inline UE::NNE::RuntimeBasic::FModelBuilder::EActivationFunction GetBuilderActivationFunction(const ELearningAgentsActivationFunction ActivationFunction)
-	{
-		switch (ActivationFunction)
-		{
-		case ELearningAgentsActivationFunction::ReLU: return UE::NNE::RuntimeBasic::FModelBuilder::EActivationFunction::ReLU;
-		case ELearningAgentsActivationFunction::ELU: return UE::NNE::RuntimeBasic::FModelBuilder::EActivationFunction::ELU;
-		case ELearningAgentsActivationFunction::TanH: return UE::NNE::RuntimeBasic::FModelBuilder::EActivationFunction::TanH;
-		default:
-			UE_LOG(LogLearning, Error, TEXT("Unknown Activation Function"));
-			return UE::NNE::RuntimeBasic::FModelBuilder::EActivationFunction::ReLU;
-		}
-	}
-}
-
 ULearningAgentsCritic::ULearningAgentsCritic() : Super(FObjectInitializer::Get()) {}
 ULearningAgentsCritic::ULearningAgentsCritic(FVTableHelper& Helper) : Super(Helper) {}
 ULearningAgentsCritic::~ULearningAgentsCritic() = default;
@@ -186,7 +170,7 @@ void ULearningAgentsCritic::SetupCritic(
 				1,
 				CriticSettings.HiddenLayerSize,
 				CriticSettings.HiddenLayerNum + 2, // Add 2 to account for input and output layers
-				UE::Learning::Agents::Critic::Private::GetBuilderActivationFunction(CriticSettings.ActivationFunction)));
+				UE::Learning::Agents::NeuralNetwork::GetBuilderActivationFunction(CriticSettings.ActivationFunction)));
 
 		check(CriticInputSize == ObservationEncodedVectorSize + MemoryStateSize);
 		check(CriticOutputSize == 1);
diff --git a/Engine/Plugins/Experimental/LearningAgents/Source/LearningAgents/Private/LearningAgentsNeuralNetwork.cpp b/Engine/Plugins/Experimental/LearningAgents/Source/LearningAgents/Private/LearningAgentsNeuralNetwork.cpp
--- a/Engine/Plugins/Experimental/LearningAgents/Source/LearningAgents/Private/LearningAgentsNeuralNetwork.cpp
+++ b/Engine/Plugins/Experimental/LearningAgents/Source/LearningAgents/Private/LearningAgentsNeuralNetwork.cpp
@@ -2,7 +2,6 @@
 
 #include "LearningAgentsNeuralNetwork.h"
 
-#include "LearningLog.h"
 #include "LearningArray.h"
 #include "LearningNeuralNetwork.h"
 
diff --git a/Engine/Plugins/Experimental/LearningAgents/Source/LearningAgents/Private/LearningAgentsObservations.cpp b/Engine/Plugins/Experimental/LearningAgents/Source/LearningAgents/Private/LearningAgentsObservations.cpp
--- a/Engine/Plugins/Experimental/LearningAgents/Source/LearningAgents/Private/LearningAgentsObservations.cpp
+++ b/Engine/Plugins/Experimental/LearningAgents/Source/LearningAgents/Private/LearningAgentsObservations.cpp
@@ -622,6 +622,17 @@ namespace UE::Learning::Agents::Observation::Private
 		default: checkNoEntry(); return Learning::Observation::EEncodingActivationFunction::ReLU;
 		}
 	}
+	
+	static inline Learning::Observation::ENormalization GetNormalizationOption(const ELearningAgentsNormalization NormalizationOption)
+	{
+		switch (NormalizationOption)
+		{
+		case ELearningAgentsNormalization::Manual: return Learning::Observation::ENormalization::Manual;
+		case ELearningAgentsNormalization::AutoShared: return Learning::Observation::ENormalization::AutoShared;
+		case ELearningAgentsNormalization::AutoPerDimension: return Learning::Observation::ENormalization::AutoPerDimension;
+		default: checkNoEntry(); return Learning::Observation::ENormalization::Manual;
+		}
+	}
 }
 
 FTransform ULearningAgentsObservations::ProjectTransformOntoGroundPlane(const FTransform Transform, const FVector LocalForwardVector, const float GroundPlaneHeight)
@@ -675,7 +686,7 @@ FLearningAgentsObservationSchemaElement ULearningAgentsObservations::SpecifyNull
 	return { Schema->ObservationSchema.CreateNull(Tag)};
 }
 
-FLearningAgentsObservationSchemaElement ULearningAgentsObservations::SpecifyContinuousObservation(ULearningAgentsObservationSchema* Schema, const int32 Size, const float Scale, const bool bActive, const FName Tag)
+FLearningAgentsObservationSchemaElement ULearningAgentsObservations::SpecifyContinuousObservation(ULearningAgentsObservationSchema* Schema, const int32 Size, const FLearningAgentsObservationNormalizationSettings& NormalizationSettings, const bool bActive, const FName Tag)
 {
 	if (!bActive)
 	{
@@ -693,19 +704,13 @@ FLearningAgentsObservationSchemaElement ULearningAgentsObservations::SpecifyCont
 		UE_LOG(LogLearning, Error, TEXT("SpecifyContinuousObservation: Invalid Continuous Observation Size '%i'."), Size);
 		return FLearningAgentsObservationSchemaElement();
 	}
-	
-	if (Scale < UE_SMALL_NUMBER)
-	{
-		UE_LOG(LogLearning, Error, TEXT("SpecifyContinuousObservation: Invalid Scale for Continuous Observation '%f', must be greater than '%f'."), Scale, UE_SMALL_NUMBER);
-		return FLearningAgentsObservationSchemaElement();
-	}
 
 	if (Size == 0)
 	{
 		UE_LOG(LogLearning, Warning, TEXT("SpecifyContinuousObservation: Specifying zero-sized Continuous Observation."));
 	}
 
-	return { Schema->ObservationSchema.CreateContinuous({ Size, Scale }, Tag) };
+	return { Schema->ObservationSchema.CreateContinuous({ Size, NormalizationSettings.Scales, NormalizationSettings.Offsets, UE::Learning::Agents::Observation::Private::GetNormalizationOption(NormalizationSettings.NormalizationOption)}, Tag) };
 }
 
 FLearningAgentsObservationSchemaElement ULearningAgentsObservations::SpecifyExclusiveDiscreteObservation(ULearningAgentsObservationSchema* Schema, const int32 Size, const bool bActive, const FName Tag)
@@ -839,7 +844,7 @@ FLearningAgentsObservationSchemaElement ULearningAgentsObservations::SpecifyCoun
 		return SpecifyNullObservation(Schema, Tag);
 	}
 
-	return SpecifyContinuousObservation(Schema, 1, 1.0f, bActive, Tag);
+	return SpecifyContinuousObservation(Schema, 1, FLearningAgentsObservationNormalizationSettings(), bActive, Tag);
 }
 
 FLearningAgentsObservationSchemaElement ULearningAgentsObservations::SpecifyStructObservation(ULearningAgentsObservationSchema* Schema, const TMap<FName, FLearningAgentsObservationSchemaElement>& Elements, const bool bActive, const FName Tag)
@@ -1302,7 +1307,7 @@ FLearningAgentsObservationSchemaElement ULearningAgentsObservations::SpecifyEnum
 		return FLearningAgentsObservationSchemaElement();
 	}
 
-	return SpecifyContinuousObservation(Schema, Enum->NumEnums() - 1, 1.0f, bActive, Tag);
+	return SpecifyContinuousObservation(Schema, Enum->NumEnums() - 1, FLearningAgentsObservationNormalizationSettings(), bActive, Tag);
 }
 
 FLearningAgentsObservationSchemaElement ULearningAgentsObservations::SpecifyBitmaskObservation(ULearningAgentsObservationSchema* Schema, const UEnum* Enum, const bool bActive, const FName Tag)
@@ -1324,7 +1329,7 @@ FLearningAgentsObservationSchemaElement ULearningAgentsObservations::SpecifyBitm
 		return FLearningAgentsObservationSchemaElement();
 	}
 
-	return SpecifyContinuousObservation(Schema, Enum->NumEnums() - 1, 1.0f, bActive, Tag);
+	return SpecifyContinuousObservation(Schema, Enum->NumEnums() - 1, FLearningAgentsObservationNormalizationSettings(), bActive, Tag);
 }
 
 FLearningAgentsObservationSchemaElement ULearningAgentsObservations::SpecifyOptionalObservation(ULearningAgentsObservationSchema* Schema, const FLearningAgentsObservationSchemaElement Element, const int32 EncodingSize, const bool bActive, const FName Tag)
@@ -1388,7 +1393,7 @@ FLearningAgentsObservationSchemaElement ULearningAgentsObservations::SpecifyBool
 		return SpecifyNullObservation(Schema, Tag);
 	}
 
-	return SpecifyContinuousObservation(Schema, 1, 1.0f, bActive, Tag);
+	return SpecifyContinuousObservation(Schema, 1, FLearningAgentsObservationNormalizationSettings(), bActive, Tag);
 }
 
 FLearningAgentsObservationSchemaElement ULearningAgentsObservations::SpecifyConv1dObservation(ULearningAgentsObservationSchema* Schema, const FLearningAgentsObservationSchemaElement Element, const FConv1dObservationParams& Conv1dParams, const bool bActive, const FName Tag)
@@ -1522,24 +1527,34 @@ FLearningAgentsObservationSchemaElement ULearningAgentsObservations::SpecifyConv
 	Tag) };
 }
 
-FLearningAgentsObservationSchemaElement ULearningAgentsObservations::SpecifyFloatObservation(ULearningAgentsObservationSchema* Schema, const float FloatScale, const bool bActive, const FName Tag)
+FLearningAgentsObservationSchemaElement ULearningAgentsObservations::SpecifyFloatObservation(ULearningAgentsObservationSchema* Schema, const FLearningAgentsFloatObservationNormalizationSettings& NormalizationSettings, const bool bActive, const FName Tag)
 {
 	if (!bActive)
 	{
 		return SpecifyNullObservation(Schema, Tag);
 	}
 
-	return SpecifyContinuousObservation(Schema, 1, FloatScale, bActive, Tag);
+	FLearningAgentsObservationNormalizationSettings Settings;
+	Settings.Scales = { NormalizationSettings.Scale };
+	Settings.Offsets = { NormalizationSettings.Offset };
+	Settings.NormalizationOption = NormalizationSettings.NormalizationOption;
+
+	return SpecifyContinuousObservation(Schema, 1, Settings, bActive, Tag);
 }
 
-FLearningAgentsObservationSchemaElement ULearningAgentsObservations::SpecifyLocationObservation(ULearningAgentsObservationSchema* Schema, const float LocationScale, const bool bActive, const FName Tag)
+FLearningAgentsObservationSchemaElement ULearningAgentsObservations::SpecifyLocationObservation(ULearningAgentsObservationSchema* Schema, const FLearningAgentsVectorObservationNormalizationSettings& NormalizationSettings, const bool bActive, const FName Tag)
 {
 	if (!bActive)
 	{
 		return SpecifyNullObservation(Schema, Tag);
 	}
 
-	return SpecifyContinuousObservation(Schema, 3, LocationScale, bActive, Tag);
+	FLearningAgentsObservationNormalizationSettings Settings;
+	Settings.Scales = { (float)NormalizationSettings.Scales.X, (float)NormalizationSettings.Scales.Y, (float)NormalizationSettings.Scales.Z };
+	Settings.Offsets = { (float)NormalizationSettings.Offsets.X, (float)NormalizationSettings.Offsets.Y, (float)NormalizationSettings.Offsets.Z };
+	Settings.NormalizationOption = NormalizationSettings.NormalizationOption;
+
+	return SpecifyContinuousObservation(Schema, 3, Settings, bActive, Tag);
 }
 
 FLearningAgentsObservationSchemaElement ULearningAgentsObservations::SpecifyRotationObservation(ULearningAgentsObservationSchema* Schema, const bool bActive, const FName Tag)
@@ -1549,7 +1564,7 @@ FLearningAgentsObservationSchemaElement ULearningAgentsObservations::SpecifyRota
 		return SpecifyNullObservation(Schema, Tag);
 	}
 
-	return SpecifyContinuousObservation(Schema, 6, 1.0f, bActive, Tag);
+	return SpecifyContinuousObservation(Schema, 6, FLearningAgentsObservationNormalizationSettings(), bActive, Tag);
 }
 
 FLearningAgentsObservationSchemaElement ULearningAgentsObservations::SpecifyScaleObservation(ULearningAgentsObservationSchema* Schema, const bool bActive, const FName Tag)
@@ -1559,10 +1574,10 @@ FLearningAgentsObservationSchemaElement ULearningAgentsObservations::SpecifyScal
 		return SpecifyNullObservation(Schema, Tag);
 	}
 
-	return SpecifyContinuousObservation(Schema, 3, 1.0f, bActive, Tag);
+	return SpecifyContinuousObservation(Schema, 3, FLearningAgentsObservationNormalizationSettings(), bActive, Tag);
 }
 
-FLearningAgentsObservationSchemaElement ULearningAgentsObservations::SpecifyTransformObservation(ULearningAgentsObservationSchema* Schema, const float LocationScale, const bool bActive, const FName Tag)
+FLearningAgentsObservationSchemaElement ULearningAgentsObservations::SpecifyTransformObservation(ULearningAgentsObservationSchema* Schema, const FLearningAgentsVectorObservationNormalizationSettings& NormalizationSettings, const bool bActive, const FName Tag)
 {
 	if (!bActive)
 	{
@@ -1576,7 +1591,7 @@ FLearningAgentsObservationSchemaElement ULearningAgentsObservations::SpecifyTran
 			TEXT("Scale")
 		},
 		{
-			SpecifyLocationObservation(Schema, LocationScale),
+			SpecifyLocationObservation(Schema, NormalizationSettings, true, TEXT("Location")),
 			SpecifyRotationObservation(Schema),
 			SpecifyScaleObservation(Schema)
 		},
@@ -1590,17 +1605,22 @@ FLearningAgentsObservationSchemaElement ULearningAgentsObservations::SpecifyAngl
 		return SpecifyNullObservation(Schema, Tag);
 	}
 
-	return SpecifyContinuousObservation(Schema, 2, 1.0f, bActive, Tag);
+	return SpecifyContinuousObservation(Schema, 2, FLearningAgentsObservationNormalizationSettings(), bActive, Tag);
 }
 
-FLearningAgentsObservationSchemaElement ULearningAgentsObservations::SpecifyVelocityObservation(ULearningAgentsObservationSchema* Schema, const float VelocityScale, const bool bActive, const FName Tag)
+FLearningAgentsObservationSchemaElement ULearningAgentsObservations::SpecifyVelocityObservation(ULearningAgentsObservationSchema* Schema, const FLearningAgentsVectorObservationNormalizationSettings& NormalizationSettings, const bool bActive, const FName Tag)
 {
 	if (!bActive)
 	{
 		return SpecifyNullObservation(Schema, Tag);
 	}
 
-	return SpecifyContinuousObservation(Schema, 3, VelocityScale, bActive, Tag);
+	FLearningAgentsObservationNormalizationSettings Settings;
+	Settings.Scales = { (float)NormalizationSettings.Scales.X, (float)NormalizationSettings.Scales.Y, (float)NormalizationSettings.Scales.Z };
+	Settings.Offsets = { (float)NormalizationSettings.Offsets.X, (float)NormalizationSettings.Offsets.Y, (float)NormalizationSettings.Offsets.Z };
+	Settings.NormalizationOption = NormalizationSettings.NormalizationOption;
+
+	return SpecifyContinuousObservation(Schema, 3, Settings, bActive, Tag);
 }
 
 FLearningAgentsObservationSchemaElement ULearningAgentsObservations::SpecifyDirectionObservation(ULearningAgentsObservationSchema* Schema, const bool bActive, const FName Tag)
@@ -1610,17 +1630,17 @@ FLearningAgentsObservationSchemaElement ULearningAgentsObservations::SpecifyDire
 		return SpecifyNullObservation(Schema, Tag);
 	}
 
-	return SpecifyContinuousObservation(Schema, 3, 1.0f, bActive, Tag);
+	return SpecifyContinuousObservation(Schema, 3, FLearningAgentsObservationNormalizationSettings(), bActive, Tag);
 }
 
-FLearningAgentsObservationSchemaElement ULearningAgentsObservations::SpecifyLocationAlongSplineObservation(ULearningAgentsObservationSchema* Schema, const float LocationScale, const bool bActive, const FName Tag)
+FLearningAgentsObservationSchemaElement ULearningAgentsObservations::SpecifyLocationAlongSplineObservation(ULearningAgentsObservationSchema* Schema, const FLearningAgentsVectorObservationNormalizationSettings& NormalizationSettings, const bool bActive, const FName Tag)
 {
 	if (!bActive)
 	{
 		return SpecifyNullObservation(Schema, Tag);
 	}
 
-	return SpecifyLocationObservation(Schema, LocationScale, bActive, Tag);
+	return SpecifyLocationObservation(Schema, NormalizationSettings, bActive, Tag);
 }
 
 FLearningAgentsObservationSchemaElement ULearningAgentsObservations::SpecifyProportionAlongSplineObservation(ULearningAgentsObservationSchema* Schema, const bool bActive, const FName Tag)
@@ -1660,7 +1680,7 @@ FLearningAgentsObservationSchemaElement ULearningAgentsObservations::SpecifyProp
 		return SpecifyNullObservation(Schema, Tag);
 	}
 
-	return SpecifyFloatObservation(Schema, 1.0f, bActive, Tag);
+	return SpecifyFloatObservation(Schema, FLearningAgentsFloatObservationNormalizationSettings(), bActive, Tag);
 }
 
 FLearningAgentsObservationSchemaElement ULearningAgentsObservations::SpecifyLidarObservation(ULearningAgentsObservationSchema* Schema, const bool bActive, const FName Tag)
@@ -1670,7 +1690,7 @@ FLearningAgentsObservationSchemaElement ULearningAgentsObservations::SpecifyLida
 		return SpecifyNullObservation(Schema, Tag);
 	}
 
-	return SpecifyFloatObservation(Schema, 1.0f, bActive, Tag);
+	return SpecifyFloatObservation(Schema, FLearningAgentsFloatObservationNormalizationSettings(), bActive, Tag);
 }
 
 void ULearningAgentsObservations::LogObservation(const ULearningAgentsObservationObject* Object, const FLearningAgentsObservationObjectElement Element)
@@ -1696,34 +1716,26 @@ FLearningAgentsObservationObjectElement ULearningAgentsObservations::MakeNullObs
 }
 
 FLearningAgentsObservationObjectElement ULearningAgentsObservations::MakeContinuousObservation(
-	ULearningAgentsObservationObject* Object, 
+	ULearningAgentsObservationObject* Object,
 	const TArray<float>& Values,
 	const bool bActive,
 	const FName Tag,
-	const bool bVisualLoggerEnabled,
-	ULearningAgentsManagerListener* VisualLoggerListener,
-	const int32 VisualLoggerAgentId,
-	const FVector VisualLoggerLocation,
-	const FLinearColor VisualLoggerColor)
+	const FLearningAgentsObservationVisualLoggerSettings& VisualLoggerSettings)
 {
 	if (!bActive)
 	{
 		return MakeNullObservation(Object, Tag);
 	}
 
-	return MakeContinuousObservationFromArrayView(Object, Values, bActive, Tag, bVisualLoggerEnabled, VisualLoggerListener, VisualLoggerAgentId, VisualLoggerLocation, VisualLoggerColor);
+	return MakeContinuousObservationFromArrayView(Object, Values, bActive, Tag, VisualLoggerSettings);
 }
 
 FLearningAgentsObservationObjectElement ULearningAgentsObservations::MakeContinuousObservationFromArrayView(
-	ULearningAgentsObservationObject* Object, 
-	const TArrayView<const float> Values, 
+	ULearningAgentsObservationObject* Object,
+	const TArrayView<const float> Values,
 	const bool bActive,
 	const FName Tag,
-	const bool bVisualLoggerEnabled,
-	ULearningAgentsManagerListener* VisualLoggerListener,
-	const int32 VisualLoggerAgentId,
-	const FVector VisualLoggerLocation,
-	const FLinearColor VisualLoggerColor)
+	const FLearningAgentsObservationVisualLoggerSettings& VisualLoggerSettings)
 {
 	if (!bActive)
 	{
@@ -1742,16 +1754,16 @@ FLearningAgentsObservationObjectElement ULearningAgentsObservations::MakeContinu
 	}
 
 #if UE_LEARNING_AGENTS_ENABLE_VISUAL_LOG
-	if (bVisualLoggerEnabled && VisualLoggerListener)
+	if (VisualLoggerSettings.bEnabled && VisualLoggerSettings.Listener)
 	{
-		const ULearningAgentsVisualLoggerObject* VisualLoggerObject = VisualLoggerListener->GetOrAddVisualLoggerObject(Tag);
+		const ULearningAgentsVisualLoggerObject* VisualLoggerObject = VisualLoggerSettings.Listener->GetOrAddVisualLoggerObject(Tag);
 
-		UE_LEARNING_AGENTS_VLOG_STRING(VisualLoggerObject, LogLearning, Display, VisualLoggerLocation,
-			VisualLoggerColor.ToFColor(true),
+		UE_LEARNING_AGENTS_VLOG_STRING(VisualLoggerObject, LogLearning, Display, VisualLoggerSettings.Location,
+			VisualLoggerSettings.Color.ToFColor(true),
 			TEXT("Listener: %s\nTag: %s\nAgent Id: % 3i\nValues: %s"),
-			*VisualLoggerListener->GetName(),
+			*VisualLoggerSettings.Listener->GetName(),
 			*Tag.ToString(),
-			VisualLoggerAgentId,
+			VisualLoggerSettings.AgentId,
 			*UE::Learning::Array::FormatFloat(Values));
 	}
 #endif
@@ -1764,11 +1776,7 @@ FLearningAgentsObservationObjectElement ULearningAgentsObservations::MakeExclusi
 	const int32 DiscreteIndex,
 	const bool bActive,
 	const FName Tag,
-	const bool bVisualLoggerEnabled,
-	ULearningAgentsManagerListener* VisualLoggerListener,
-	const int32 VisualLoggerAgentId,
-	const FVector VisualLoggerLocation,
-	const FLinearColor VisualLoggerColor)
+	const FLearningAgentsObservationVisualLoggerSettings& VisualLoggerSettings)
 {
 	if (!bActive)
 	{
@@ -1782,16 +1790,16 @@ FLearningAgentsObservationObjectElement ULearningAgentsObservations::MakeExclusi
 	}
 
 #if UE_LEARNING_AGENTS_ENABLE_VISUAL_LOG
-	if (bVisualLoggerEnabled && VisualLoggerListener)
+	if (VisualLoggerSettings.bEnabled && VisualLoggerSettings.Listener)
 	{
-		const ULearningAgentsVisualLoggerObject* VisualLoggerObject = VisualLoggerListener->GetOrAddVisualLoggerObject(Tag);
+		const ULearningAgentsVisualLoggerObject* VisualLoggerObject = VisualLoggerSettings.Listener->GetOrAddVisualLoggerObject(Tag);
 
-		UE_LEARNING_AGENTS_VLOG_STRING(VisualLoggerObject, LogLearning, Display, VisualLoggerLocation,
-			VisualLoggerColor.ToFColor(true),
+		UE_LEARNING_AGENTS_VLOG_STRING(VisualLoggerObject, LogLearning, Display, VisualLoggerSettings.Location,
+			VisualLoggerSettings.Color.ToFColor(true),
 			TEXT("Listener: %s\nTag: %s\nAgent Id: % 3i\nIndex: [%i]"),
-			*VisualLoggerListener->GetName(),
+			*VisualLoggerSettings.Listener->GetName(),
 			*Tag.ToString(),
-			VisualLoggerAgentId,
+			VisualLoggerSettings.AgentId,
 			DiscreteIndex);
 	}
 #endif
@@ -1804,18 +1812,14 @@ FLearningAgentsObservationObjectElement ULearningAgentsObservations::MakeInclusi
 	const TArray<int32>& DiscreteIndices, 
 	const bool bActive,
 	const FName Tag,
-	const bool bVisualLoggerEnabled,
-	ULearningAgentsManagerListener* VisualLoggerListener,
-	const int32 VisualLoggerAgentId,
-	const FVector VisualLoggerLocation,
-	const FLinearColor VisualLoggerColor)
+	const FLearningAgentsObservationVisualLoggerSettings& VisualLoggerSettings)
 {
 	if (!bActive)
 	{
 		return MakeNullObservation(Object, Tag);
 	}
 
-	return MakeInclusiveDiscreteObservationFromArrayView(Object, DiscreteIndices, bActive, Tag, bVisualLoggerEnabled, VisualLoggerListener, VisualLoggerAgentId, VisualLoggerLocation, VisualLoggerColor);
+	return MakeInclusiveDiscreteObservationFromArrayView(Object, DiscreteIndices, bActive, Tag, VisualLoggerSettings);
 }
 
 FLearningAgentsObservationObjectElement ULearningAgentsObservations::MakeInclusiveDiscreteObservationFromArrayView(
@@ -1823,11 +1827,7 @@ FLearningAgentsObservationObjectElement ULearningAgentsObservations::MakeInclusi
 	const TArrayView<const int32> DiscreteIndices, 
 	const bool bActive,
 	const FName Tag,
-	const bool bVisualLoggerEnabled,
-	ULearningAgentsManagerListener* VisualLoggerListener,
-	const int32 VisualLoggerAgentId,
-	const FVector VisualLoggerLocation,
-	const FLinearColor VisualLoggerColor)
+	const FLearningAgentsObservationVisualLoggerSettings& VisualLoggerSettings)
 {
 	if (!bActive)
 	{
@@ -1847,16 +1847,16 @@ FLearningAgentsObservationObjectElement ULearningAgentsObservations::MakeInclusi
 	}
 
 #if UE_LEARNING_AGENTS_ENABLE_VISUAL_LOG
-	if (bVisualLoggerEnabled && VisualLoggerListener)
+	if (VisualLoggerSettings.bEnabled && VisualLoggerSettings.Listener)
 	{
-		const ULearningAgentsVisualLoggerObject* VisualLoggerObject = VisualLoggerListener->GetOrAddVisualLoggerObject(Tag);
+		const ULearningAgentsVisualLoggerObject* VisualLoggerObject = VisualLoggerSettings.Listener->GetOrAddVisualLoggerObject(Tag);
 
-		UE_LEARNING_AGENTS_VLOG_STRING(VisualLoggerObject, LogLearning, Display, VisualLoggerLocation,
-			VisualLoggerColor.ToFColor(true),
+		UE_LEARNING_AGENTS_VLOG_STRING(VisualLoggerObject, LogLearning, Display, VisualLoggerSettings.Location,
+			VisualLoggerSettings.Color.ToFColor(true),
 			TEXT("Listener: %s\nTag: %s\nAgent Id: % 3i\nIndices: %s"),
-			*VisualLoggerListener->GetName(),
+			*VisualLoggerSettings.Listener->GetName(),
 			*Tag.ToString(),
-			VisualLoggerAgentId,
+			VisualLoggerSettings.AgentId,
 			*UE::Learning::Array::FormatInt32(DiscreteIndices));
 	}
 #endif
@@ -1869,11 +1869,7 @@ FLearningAgentsObservationObjectElement ULearningAgentsObservations::MakeNamedEx
 	const FName ElementName,
 	const bool bActive,
 	const FName Tag,
-	const bool bVisualLoggerEnabled,
-	ULearningAgentsManagerListener* VisualLoggerListener,
-	const int32 VisualLoggerAgentId,
-	const FVector VisualLoggerLocation,
-	const FLinearColor VisualLoggerColor)
+	const FLearningAgentsObservationVisualLoggerSettings& VisualLoggerSettings)
 {
 	if (!bActive)
 	{
@@ -1887,16 +1883,16 @@ FLearningAgentsObservationObjectElement ULearningAgentsObservations::MakeNamedEx
 	}
 
 #if UE_LEARNING_AGENTS_ENABLE_VISUAL_LOG
-	if (bVisualLoggerEnabled && VisualLoggerListener)
+	if (VisualLoggerSettings.bEnabled && VisualLoggerSettings.Listener)
 	{
-		const ULearningAgentsVisualLoggerObject* VisualLoggerObject = VisualLoggerListener->GetOrAddVisualLoggerObject(Tag);
+		const ULearningAgentsVisualLoggerObject* VisualLoggerObject = VisualLoggerSettings.Listener->GetOrAddVisualLoggerObject(Tag);
 
-		UE_LEARNING_AGENTS_VLOG_STRING(VisualLoggerObject, LogLearning, Display, VisualLoggerLocation,
-			VisualLoggerColor.ToFColor(true),
+		UE_LEARNING_AGENTS_VLOG_STRING(VisualLoggerObject, LogLearning, Display, VisualLoggerSettings.Location,
+			VisualLoggerSettings.Color.ToFColor(true),
 			TEXT("Listener: %s\nTag: %s\nAgent Id: % 3i\nName: [%s]"),
-			*VisualLoggerListener->GetName(),
+			*VisualLoggerSettings.Listener->GetName(),
 			*Tag.ToString(),
-			VisualLoggerAgentId,
+			VisualLoggerSettings.AgentId,
 			*ElementName.ToString());
 	}
 #endif
@@ -1909,18 +1905,14 @@ FLearningAgentsObservationObjectElement ULearningAgentsObservations::MakeNamedIn
 	const TArray<FName>& ElementNames,
 	const bool bActive,
 	const FName Tag,
-	const bool bVisualLoggerEnabled,
-	ULearningAgentsManagerListener* VisualLoggerListener,
-	const int32 VisualLoggerAgentId,
-	const FVector VisualLoggerLocation,
-	const FLinearColor VisualLoggerColor)
+	const FLearningAgentsObservationVisualLoggerSettings& VisualLoggerSettings)
 {
 	if (!bActive)
 	{
 		return MakeNullObservation(Object, Tag);
 	}
 
-	return MakeNamedInclusiveDiscreteObservationFromArrayView(Object, ElementNames, bActive, Tag, bVisualLoggerEnabled, VisualLoggerListener, VisualLoggerAgentId, VisualLoggerLocation, VisualLoggerColor);
+	return MakeNamedInclusiveDiscreteObservationFromArrayView(Object, ElementNames, bActive, Tag, VisualLoggerSettings);
 }
 
 FLearningAgentsObservationObjectElement ULearningAgentsObservations::MakeNamedInclusiveDiscreteObservationFromArrayView(
@@ -1928,11 +1920,7 @@ FLearningAgentsObservationObjectElement ULearningAgentsObservations::MakeNamedIn
 	const TArrayView<const FName> ElementNames,
 	const bool bActive,
 	const FName Tag,
-	const bool bVisualLoggerEnabled,
-	ULearningAgentsManagerListener* VisualLoggerListener,
-	const int32 VisualLoggerAgentId,
-	const FVector VisualLoggerLocation,
-	const FLinearColor VisualLoggerColor)
+	const FLearningAgentsObservationVisualLoggerSettings& VisualLoggerSettings)
 {
 	if (!bActive)
 	{
@@ -1952,16 +1940,16 @@ FLearningAgentsObservationObjectElement ULearningAgentsObservations::MakeNamedIn
 	}
 
 #if UE_LEARNING_AGENTS_ENABLE_VISUAL_LOG
-	if (bVisualLoggerEnabled && VisualLoggerListener)
+	if (VisualLoggerSettings.bEnabled && VisualLoggerSettings.Listener)
 	{
-		const ULearningAgentsVisualLoggerObject* VisualLoggerObject = VisualLoggerListener->GetOrAddVisualLoggerObject(Tag);
+		const ULearningAgentsVisualLoggerObject* VisualLoggerObject = VisualLoggerSettings.Listener->GetOrAddVisualLoggerObject(Tag);
 
-		UE_LEARNING_AGENTS_VLOG_STRING(VisualLoggerObject, LogLearning, Display, VisualLoggerLocation,
-			VisualLoggerColor.ToFColor(true),
+		UE_LEARNING_AGENTS_VLOG_STRING(VisualLoggerObject, LogLearning, Display, VisualLoggerSettings.Location,
+			VisualLoggerSettings.Color.ToFColor(true),
 			TEXT("Listener: %s\nTag: %s\nAgent Id: % 3i\nNames: %s"),
-			*VisualLoggerListener->GetName(),
+			*VisualLoggerSettings.Listener->GetName(),
 			*Tag.ToString(),
-			VisualLoggerAgentId,
+			VisualLoggerSettings.AgentId,
 			*UE::Learning::Agents::Observation::Private::FormatNames(ElementNames, 256));
 	}
 #endif
@@ -1975,11 +1963,7 @@ FLearningAgentsObservationObjectElement ULearningAgentsObservations::MakeCountOb
 	const int32 MaxNum, 
 	const bool bActive,
 	const FName Tag,
-	const bool bVisualLoggerEnabled,
-	ULearningAgentsManagerListener* VisualLoggerListener,
-	const int32 VisualLoggerAgentId,
-	const FVector VisualLoggerLocation,
-	const FLinearColor VisualLoggerColor)
+	const FLearningAgentsObservationVisualLoggerSettings& VisualLoggerSettings)
 {
 	if (!bActive)
 	{
@@ -1995,16 +1979,16 @@ FLearningAgentsObservationObjectElement ULearningAgentsObservations::MakeCountOb
 	const float Encoded = (float)Num / (float)MaxNum;
 
 #if UE_LEARNING_AGENTS_ENABLE_VISUAL_LOG
-	if (bVisualLoggerEnabled && VisualLoggerListener)
+	if (VisualLoggerSettings.bEnabled && VisualLoggerSettings.Listener)
 	{
-		const ULearningAgentsVisualLoggerObject* VisualLoggerObject = VisualLoggerListener->GetOrAddVisualLoggerObject(Tag);
+		const ULearningAgentsVisualLoggerObject* VisualLoggerObject = VisualLoggerSettings.Listener->GetOrAddVisualLoggerObject(Tag);
 
-		UE_LEARNING_AGENTS_VLOG_STRING(VisualLoggerObject, LogLearning, Display, VisualLoggerLocation,
-			VisualLoggerColor.ToFColor(true),
+		UE_LEARNING_AGENTS_VLOG_STRING(VisualLoggerObject, LogLearning, Display, VisualLoggerSettings.Location,
+			VisualLoggerSettings.Color.ToFColor(true),
 			TEXT("Listener: %s\nTag: %s\nAgent Id: % 3i\nNum: [%i]\nMax Num: [%i]\nEncoded: [%6.2f]"),
-			*VisualLoggerListener->GetName(),
+			*VisualLoggerSettings.Listener->GetName(),
 			*Tag.ToString(),
-			VisualLoggerAgentId,
+			VisualLoggerSettings.AgentId,
 			Num,
 			MaxNum,
 			Encoded);
@@ -2446,11 +2430,7 @@ FLearningAgentsObservationObjectElement ULearningAgentsObservations::MakeEnumObs
 	const uint8 EnumValue, 
 	const bool bActive,
 	const FName Tag,
-	const bool bVisualLoggerEnabled,
-	ULearningAgentsManagerListener* VisualLoggerListener,
-	const int32 VisualLoggerAgentId,
-	const FVector VisualLoggerLocation,
-	const FLinearColor VisualLoggerColor)
+	const FLearningAgentsObservationVisualLoggerSettings& VisualLoggerSettings)
 {
 	if (!bActive)
 	{
@@ -2476,16 +2456,16 @@ FLearningAgentsObservationObjectElement ULearningAgentsObservations::MakeEnumObs
 	OneHot[EnumValueIndex] = 1.0f;
 
 #if UE_LEARNING_AGENTS_ENABLE_VISUAL_LOG
-	if (bVisualLoggerEnabled && VisualLoggerListener)
+	if (VisualLoggerSettings.bEnabled && VisualLoggerSettings.Listener)
 	{
-		const ULearningAgentsVisualLoggerObject* VisualLoggerObject = VisualLoggerListener->GetOrAddVisualLoggerObject(Tag);
+		const ULearningAgentsVisualLoggerObject* VisualLoggerObject = VisualLoggerSettings.Listener->GetOrAddVisualLoggerObject(Tag);
 
-		UE_LEARNING_AGENTS_VLOG_STRING(VisualLoggerObject, LogLearning, Display, VisualLoggerLocation,
-			VisualLoggerColor.ToFColor(true),
+		UE_LEARNING_AGENTS_VLOG_STRING(VisualLoggerObject, LogLearning, Display, VisualLoggerSettings.Location,
+			VisualLoggerSettings.Color.ToFColor(true),
 			TEXT("Listener: %s\nTag: %s\nAgent Id: % 3i\nEnum: %s\nSize: [%i]\nValue: [%s]\nIndex: [%i]\nEncoded: %s"),
-			*VisualLoggerListener->GetName(),
+			*VisualLoggerSettings.Listener->GetName(),
 			*Tag.ToString(),
-			VisualLoggerAgentId,
+			VisualLoggerSettings.AgentId,
 			*Enum->GetName(),
 			Enum->NumEnums() - 1,
 			*Enum->GetDisplayNameTextByValue(EnumValue).ToString(),
@@ -2503,11 +2483,7 @@ FLearningAgentsObservationObjectElement ULearningAgentsObservations::MakeBitmask
 	const int32 BitmaskValue, 
 	const bool bActive,
 	const FName Tag,
-	const bool bVisualLoggerEnabled,
-	ULearningAgentsManagerListener* VisualLoggerListener,
-	const int32 VisualLoggerAgentId,
-	const FVector VisualLoggerLocation,
-	const FLinearColor VisualLoggerColor)
+	const FLearningAgentsObservationVisualLoggerSettings& VisualLoggerSettings)
 {
 	if (!bActive)
 	{
@@ -2538,9 +2514,9 @@ FLearningAgentsObservationObjectElement ULearningAgentsObservations::MakeBitmask
 	}
 
 #if UE_LEARNING_AGENTS_ENABLE_VISUAL_LOG
-	if (bVisualLoggerEnabled && VisualLoggerListener)
+	if (VisualLoggerSettings.bEnabled && VisualLoggerSettings.Listener)
 	{
-		const ULearningAgentsVisualLoggerObject* VisualLoggerObject = VisualLoggerListener->GetOrAddVisualLoggerObject(Tag);
+		const ULearningAgentsVisualLoggerObject* VisualLoggerObject = VisualLoggerSettings.Listener->GetOrAddVisualLoggerObject(Tag);
 
 		FString ValuesString;
 		FString IndicesString;
@@ -2557,12 +2533,12 @@ FLearningAgentsObservationObjectElement ULearningAgentsObservations::MakeBitmask
 		ValuesString = ValuesString.TrimEnd();
 		IndicesString = IndicesString.TrimEnd();
 
-		UE_LEARNING_AGENTS_VLOG_STRING(VisualLoggerObject, LogLearning, Display, VisualLoggerLocation,
-			VisualLoggerColor.ToFColor(true),
+		UE_LEARNING_AGENTS_VLOG_STRING(VisualLoggerObject, LogLearning, Display, VisualLoggerSettings.Location,
+			VisualLoggerSettings.Color.ToFColor(true),
 			TEXT("Listener: %s\nTag: %s\nAgent Id: % 3i\nEnum: %s\nSize: [%i]\nValues: [%s]\nIndices: [%s]\nEncoded: %s"),
-			*VisualLoggerListener->GetName(),
+			*VisualLoggerSettings.Listener->GetName(),
 			*Tag.ToString(),
-			VisualLoggerAgentId,
+			VisualLoggerSettings.AgentId,
 			*Enum->GetName(),
 			Enum->NumEnums() - 1,
 			*ValuesString,
@@ -2665,11 +2641,7 @@ FLearningAgentsObservationObjectElement ULearningAgentsObservations::MakeBoolObs
 	const bool bValue,
 	const bool bActive,
 	const FName Tag,
-	const bool bVisualLoggerEnabled,
-	ULearningAgentsManagerListener* VisualLoggerListener,
-	const int32 VisualLoggerAgentId,
-	const FVector VisualLoggerLocation,
-	const FLinearColor VisualLoggerColor)
+	const FLearningAgentsObservationVisualLoggerSettings& VisualLoggerSettings)
 {
 	if (!bActive)
 	{
@@ -2679,16 +2651,16 @@ FLearningAgentsObservationObjectElement ULearningAgentsObservations::MakeBoolObs
 	const float EncodedBool = bValue ? 1.0f : -1.0f;
 
 #if UE_LEARNING_AGENTS_ENABLE_VISUAL_LOG
-	if (bVisualLoggerEnabled && VisualLoggerListener)
+	if (VisualLoggerSettings.bEnabled && VisualLoggerSettings.Listener)
 	{
-		const ULearningAgentsVisualLoggerObject* VisualLoggerObject = VisualLoggerListener->GetOrAddVisualLoggerObject(Tag);
+		const ULearningAgentsVisualLoggerObject* VisualLoggerObject = VisualLoggerSettings.Listener->GetOrAddVisualLoggerObject(Tag);
 
-		UE_LEARNING_AGENTS_VLOG_STRING(VisualLoggerObject, LogLearning, Display, VisualLoggerLocation,
-			VisualLoggerColor.ToFColor(true),
+		UE_LEARNING_AGENTS_VLOG_STRING(VisualLoggerObject, LogLearning, Display, VisualLoggerSettings.Location,
+			VisualLoggerSettings.Color.ToFColor(true),
 			TEXT("Listener: %s\nTag: %s\nAgent Id: % 3i\nValue: [%s]\nEncoded: [% 6.2f]"),
-			*VisualLoggerListener->GetName(),
+			*VisualLoggerSettings.Listener->GetName(),
 			*Tag.ToString(),
-			VisualLoggerAgentId,
+			VisualLoggerSettings.AgentId,
 			bValue ? TEXT("true") : TEXT("false"),
 			EncodedBool);
 	}
@@ -2746,11 +2718,7 @@ FLearningAgentsObservationObjectElement ULearningAgentsObservations::MakeFloatOb
 	const float Value,
 	const bool bActive,
 	const FName Tag,
-	const bool bVisualLoggerEnabled,
-	ULearningAgentsManagerListener* VisualLoggerListener,
-	const int32 VisualLoggerAgentId,
-	const FVector VisualLoggerLocation,
-	const FLinearColor VisualLoggerColor)
+	const FLearningAgentsObservationVisualLoggerSettings& VisualLoggerSettings)
 {
 	if (!bActive)
 	{
@@ -2758,21 +2726,21 @@ FLearningAgentsObservationObjectElement ULearningAgentsObservations::MakeFloatOb
 	}
 
 #if UE_LEARNING_AGENTS_ENABLE_VISUAL_LOG
-	if (bVisualLoggerEnabled && VisualLoggerListener)
+	if (VisualLoggerSettings.bEnabled && VisualLoggerSettings.Listener)
 	{
-		const ULearningAgentsVisualLoggerObject* VisualLoggerObject = VisualLoggerListener->GetOrAddVisualLoggerObject(Tag);
+		const ULearningAgentsVisualLoggerObject* VisualLoggerObject = VisualLoggerSettings.Listener->GetOrAddVisualLoggerObject(Tag);
 
-		UE_LEARNING_AGENTS_VLOG_STRING(VisualLoggerObject, LogLearning, Display, VisualLoggerLocation,
-			VisualLoggerColor.ToFColor(true),
+		UE_LEARNING_AGENTS_VLOG_STRING(VisualLoggerObject, LogLearning, Display, VisualLoggerSettings.Location,
+			VisualLoggerSettings.Color.ToFColor(true),
 			TEXT("Listener: %s\nTag: %s\nAgent Id: % 3i\nValue: [% 6.1f]"),
-			*VisualLoggerListener->GetName(),
+			*VisualLoggerSettings.Listener->GetName(),
 			*Tag.ToString(),
-			VisualLoggerAgentId,
+			VisualLoggerSettings.AgentId,
 			Value);
 	}
 #endif
 
-	return MakeContinuousObservationFromArrayView(Object, { Value }, bActive, Tag);
+	return MakeContinuousObservationFromArrayView(Object, { Value }, bActive, Tag, FLearningAgentsObservationVisualLoggerSettings());
 }
 
 FLearningAgentsObservationObjectElement ULearningAgentsObservations::MakeLocationObservation(
@@ -2781,11 +2749,7 @@ FLearningAgentsObservationObjectElement ULearningAgentsObservations::MakeLocatio
 	const FTransform RelativeTransform,
 	const bool bActive,
 	const FName Tag,
-	const bool bVisualLoggerEnabled,
-	ULearningAgentsManagerListener* VisualLoggerListener,
-	const int32 VisualLoggerAgentId,
-	const FVector VisualLoggerLocation,
-	const FLinearColor VisualLoggerColor)
+	const FLearningAgentsObservationVisualLoggerSettings& VisualLoggerSettings)
 {
 	if (!bActive)
 	{
@@ -2795,34 +2759,34 @@ FLearningAgentsObservationObjectElement ULearningAgentsObservations::MakeLocatio
 	const FVector LocalLocation = RelativeTransform.InverseTransformPosition(Location);
 
 #if UE_LEARNING_AGENTS_ENABLE_VISUAL_LOG
-	if (bVisualLoggerEnabled && VisualLoggerListener)
+	if (VisualLoggerSettings.bEnabled && VisualLoggerSettings.Listener)
 	{
-		const ULearningAgentsVisualLoggerObject* VisualLoggerObject = VisualLoggerListener->GetOrAddVisualLoggerObject(Tag);
+		const ULearningAgentsVisualLoggerObject* VisualLoggerObject = VisualLoggerSettings.Listener->GetOrAddVisualLoggerObject(Tag);
 
 		UE_LEARNING_AGENTS_VLOG_LOCATION(VisualLoggerObject, LogLearning, Display,
 			Location,
 			10,
-			VisualLoggerColor.ToFColor(true),
+			VisualLoggerSettings.Color.ToFColor(true),
 			TEXT(""));
 
 		UE_LEARNING_AGENTS_VLOG_SEGMENT(VisualLoggerObject, LogLearning, Display,
 			RelativeTransform.GetTranslation(),
 			Location,
-			VisualLoggerColor.ToFColor(true),
+			VisualLoggerSettings.Color.ToFColor(true),
 			TEXT(""));
 
 		UE_LEARNING_AGENTS_VLOG_TRANSFORM(VisualLoggerObject, LogLearning, Display,
 			RelativeTransform.GetTranslation(),
 			RelativeTransform.GetRotation(),
-			VisualLoggerColor.ToFColor(true),
+			VisualLoggerSettings.Color.ToFColor(true),
 			TEXT(""));
 
-		UE_LEARNING_AGENTS_VLOG_STRING(VisualLoggerObject, LogLearning, Display, VisualLoggerLocation,
-			VisualLoggerColor.ToFColor(true),
+		UE_LEARNING_AGENTS_VLOG_STRING(VisualLoggerObject, LogLearning, Display, VisualLoggerSettings.Location,
+			VisualLoggerSettings.Color.ToFColor(true),
 			TEXT("Listener: %s\nTag: %s\nAgent Id: % 3i\nLocation: [% 6.1f % 6.1f % 6.1f]\nLocal Location: [% 6.1f % 6.1f % 6.1f]"),
-			*VisualLoggerListener->GetName(),
+			*VisualLoggerSettings.Listener->GetName(),
 			*Tag.ToString(),
-			VisualLoggerAgentId,
+			VisualLoggerSettings.AgentId,
 			Location.X, Location.Y, Location.Z,
 			LocalLocation.X, LocalLocation.Y, LocalLocation.Z);
 	}
@@ -2832,7 +2796,7 @@ FLearningAgentsObservationObjectElement ULearningAgentsObservations::MakeLocatio
 		(float)LocalLocation.X,
 		(float)LocalLocation.Y,
 		(float)LocalLocation.Z,
-		}, bActive, Tag);
+		}, bActive, Tag, FLearningAgentsObservationVisualLoggerSettings());
 }
 
 FLearningAgentsObservationObjectElement ULearningAgentsObservations::MakeRotationObservation(
@@ -2841,12 +2805,7 @@ FLearningAgentsObservationObjectElement ULearningAgentsObservations::MakeRotatio
 	const FRotator RelativeRotation,
 	const bool bActive,
 	const FName Tag,
-	const bool bVisualLoggerEnabled,
-	ULearningAgentsManagerListener* VisualLoggerListener,
-	const int32 VisualLoggerAgentId,
-	const FVector VisualLoggerRotationLocation,
-	const FVector VisualLoggerLocation,
-	const FLinearColor VisualLoggerColor)
+	const FLearningAgentsObservationVisualLoggerSettings& VisualLoggerSettings)
 {
 	if (!bActive)
 	{
@@ -2858,7 +2817,7 @@ FLearningAgentsObservationObjectElement ULearningAgentsObservations::MakeRotatio
 	return MakeRotationObservationFromQuat(
 		Object,
 		FQuat::MakeFromRotator(Rotation), FQuat::MakeFromRotator(RelativeRotation), bActive, Tag,
-		bVisualLoggerEnabled, VisualLoggerListener, VisualLoggerAgentId, VisualLoggerRotationLocation, VisualLoggerLocation, VisualLoggerColor);
+		VisualLoggerSettings);
 }
 
 FLearningAgentsObservationObjectElement ULearningAgentsObservations::MakeRotationObservationFromQuat(
@@ -2867,12 +2826,7 @@ FLearningAgentsObservationObjectElement ULearningAgentsObservations::MakeRotatio
 	const FQuat RelativeRotation,
 	const bool bActive,
 	const FName Tag,
-	const bool bVisualLoggerEnabled,
-	ULearningAgentsManagerListener* VisualLoggerListener,
-	const int32 VisualLoggerAgentId,
-	const FVector VisualLoggerRotationLocation,
-	const FVector VisualLoggerLocation,
-	const FLinearColor VisualLoggerColor)
+	const FLearningAgentsObservationVisualLoggerSettings& VisualLoggerSettings)
 {
 	if (!bActive)
 	{
@@ -2884,22 +2838,22 @@ FLearningAgentsObservationObjectElement ULearningAgentsObservations::MakeRotatio
 	const FVector LocalAxisRight = LocalRotation.GetRightVector();
 
 #if UE_LEARNING_AGENTS_ENABLE_VISUAL_LOG
-	if (bVisualLoggerEnabled && VisualLoggerListener)
+	if (VisualLoggerSettings.bEnabled && VisualLoggerSettings.Listener)
 	{
-		const ULearningAgentsVisualLoggerObject* VisualLoggerObject = VisualLoggerListener->GetOrAddVisualLoggerObject(Tag);
+		const ULearningAgentsVisualLoggerObject* VisualLoggerObject = VisualLoggerSettings.Listener->GetOrAddVisualLoggerObject(Tag);
 
 		UE_LEARNING_AGENTS_VLOG_TRANSFORM(VisualLoggerObject, LogLearning, Display,
-			VisualLoggerRotationLocation,
+			VisualLoggerSettings.DebugLocation,
 			LocalRotation.Rotator(),
-			VisualLoggerColor.ToFColor(true),
+			VisualLoggerSettings.Color.ToFColor(true),
 			TEXT(""));
 
-		UE_LEARNING_AGENTS_VLOG_STRING(VisualLoggerObject, LogLearning, Display, VisualLoggerLocation,
-			VisualLoggerColor.ToFColor(true),
+		UE_LEARNING_AGENTS_VLOG_STRING(VisualLoggerObject, LogLearning, Display, VisualLoggerSettings.Location,
+			VisualLoggerSettings.Color.ToFColor(true),
 			TEXT("Listener: %s\nTag: %s\nAgent Id: % 3i\nRotation: [% 6.1f % 6.1f % 6.1f % 6.1f]\nLocal Rotation: [% 6.1f % 6.1f % 6.1f % 6.1f]\nEncoded Forward: [% 6.2f % 6.2f % 6.2f]\nEncoded Right: [% 6.2f % 6.2f % 6.2f]"),
-			*VisualLoggerListener->GetName(),
+			*VisualLoggerSettings.Listener->GetName(),
 			*Tag.ToString(),
-			VisualLoggerAgentId,
+			VisualLoggerSettings.AgentId,
 			Rotation.W, Rotation.X, Rotation.Y, Rotation.Z,
 			LocalRotation.W, LocalRotation.X, LocalRotation.Y, LocalRotation.Z,
 			LocalAxisForward.X, LocalAxisForward.Y, LocalAxisForward.Z,
@@ -2923,12 +2877,7 @@ FLearningAgentsObservationObjectElement ULearningAgentsObservations::MakeScaleOb
 	const FVector RelativeScale,
 	const bool bActive,
 	const FName Tag,
-	const bool bVisualLoggerEnabled,
-	ULearningAgentsManagerListener* VisualLoggerListener,
-	const int32 VisualLoggerAgentId,
-	const FVector VisualLoggerScaleLocation,
-	const FVector VisualLoggerLocation,
-	const FLinearColor VisualLoggerColor)
+	const FLearningAgentsObservationVisualLoggerSettings& VisualLoggerSettings)
 {
 	if (!bActive)
 	{
@@ -2940,16 +2889,16 @@ FLearningAgentsObservationObjectElement ULearningAgentsObservations::MakeScaleOb
 		UE::Learning::Agents::Observation::Private::VectorLogSafe(RelativeScale);
 
 #if UE_LEARNING_AGENTS_ENABLE_VISUAL_LOG
-	if (bVisualLoggerEnabled && VisualLoggerListener)
+	if (VisualLoggerSettings.bEnabled && VisualLoggerSettings.Listener)
 	{
-		const ULearningAgentsVisualLoggerObject* VisualLoggerObject = VisualLoggerListener->GetOrAddVisualLoggerObject(Tag);
+		const ULearningAgentsVisualLoggerObject* VisualLoggerObject = VisualLoggerSettings.Listener->GetOrAddVisualLoggerObject(Tag);
 
-		UE_LEARNING_AGENTS_VLOG_STRING(VisualLoggerObject, LogLearning, Display, VisualLoggerLocation,
-			VisualLoggerColor.ToFColor(true),
+		UE_LEARNING_AGENTS_VLOG_STRING(VisualLoggerObject, LogLearning, Display, VisualLoggerSettings.Location,
+			VisualLoggerSettings.Color.ToFColor(true),
 			TEXT("Listener: %s\nTag: %s\nAgent Id: % 3i\nScale: [% 6.1f % 6.1f % 6.1f]\nEncoded: [% 6.2f % 6.2f % 6.2f]"),
-			*VisualLoggerListener->GetName(),
+			*VisualLoggerSettings.Listener->GetName(),
 			*Tag.ToString(),
-			VisualLoggerAgentId,
+			VisualLoggerSettings.AgentId,
 			Scale.X, Scale.Y, Scale.Z,
 			LocalLogScale.X, LocalLogScale.Y, LocalLogScale.Z);
 	}
@@ -2968,11 +2917,7 @@ FLearningAgentsObservationObjectElement ULearningAgentsObservations::MakeTransfo
 	const FTransform RelativeTransform,
 	const bool bActive,
 	const FName Tag,
-	const bool bVisualLoggerEnabled,
-	ULearningAgentsManagerListener* VisualLoggerListener,
-	const int32 VisualLoggerAgentId,
-	const FVector VisualLoggerLocation,
-	const FLinearColor VisualLoggerColor)
+	const FLearningAgentsObservationVisualLoggerSettings& VisualLoggerSettings)
 {
 	if (!bActive)
 	{
@@ -2982,9 +2927,9 @@ FLearningAgentsObservationObjectElement ULearningAgentsObservations::MakeTransfo
 	const FTransform LocalTransform = Transform * RelativeTransform.Inverse();
 	
 #if UE_LEARNING_AGENTS_ENABLE_VISUAL_LOG
-	if (bVisualLoggerEnabled && VisualLoggerListener)
+	if (VisualLoggerSettings.bEnabled && VisualLoggerSettings.Listener)
 	{
-		const ULearningAgentsVisualLoggerObject* VisualLoggerObject = VisualLoggerListener->GetOrAddVisualLoggerObject(Tag);
+		const ULearningAgentsVisualLoggerObject* VisualLoggerObject = VisualLoggerSettings.Listener->GetOrAddVisualLoggerObject(Tag);
 
 		const FVector LocalLocation = LocalTransform.GetLocation();
 		const FRotator LocalRotation = LocalTransform.Rotator();
@@ -2997,15 +2942,15 @@ FLearningAgentsObservationObjectElement ULearningAgentsObservations::MakeTransfo
 		UE_LEARNING_AGENTS_VLOG_TRANSFORM(VisualLoggerObject, LogLearning, Display,
 			Location,
 			Rotation,
-			VisualLoggerColor.ToFColor(true),
+			VisualLoggerSettings.Color.ToFColor(true),
 			TEXT(""));
 
-		UE_LEARNING_AGENTS_VLOG_STRING(VisualLoggerObject, LogLearning, Display, VisualLoggerLocation,
-			VisualLoggerColor.ToFColor(true),
+		UE_LEARNING_AGENTS_VLOG_STRING(VisualLoggerObject, LogLearning, Display, VisualLoggerSettings.Location,
+			VisualLoggerSettings.Color.ToFColor(true),
 			TEXT("Listener: %s\nTag: %s\nAgent Id: % 3i\nLocation: [% 6.1f % 6.1f % 6.1f]\nLocal Location: [% 6.1f % 6.1f % 6.1f]\nRotation: [% 6.1f % 6.1f % 6.1f]\nLocal Rotation: [% 6.1f % 6.1f % 6.1f]\nScale: [% 6.1f % 6.1f % 6.1f]\nLocal Scale: [% 6.1f % 6.1f % 6.1f]"),
-			*VisualLoggerListener->GetName(),
+			*VisualLoggerSettings.Listener->GetName(),
 			*Tag.ToString(),
-			VisualLoggerAgentId,
+			VisualLoggerSettings.AgentId,
 			Location.X, Location.Y, Location.Z,
 			LocalLocation.X, LocalLocation.Y, LocalLocation.Z,
 			Rotation.Roll, Rotation.Pitch, Rotation.Yaw,
@@ -3035,11 +2980,7 @@ FLearningAgentsObservationObjectElement ULearningAgentsObservations::MakeAngleOb
 	const float RelativeAngle, 
 	const bool bActive,
 	const FName Tag,
-	const bool bVisualLoggerEnabled,
-	ULearningAgentsManagerListener* VisualLoggerListener,
-	const int32 VisualLoggerAgentId,
-	const FVector VisualLoggerLocation,
-	const FLinearColor VisualLoggerColor)
+	const FLearningAgentsObservationVisualLoggerSettings& VisualLoggerSettings)
 {
 	if (!bActive)
 	{
@@ -3051,24 +2992,24 @@ FLearningAgentsObservationObjectElement ULearningAgentsObservations::MakeAngleOb
 	const float EncodedY = FMath::Cos(FMath::DegreesToRadians(LocalAngle));
 
 #if UE_LEARNING_AGENTS_ENABLE_VISUAL_LOG
-	if (bVisualLoggerEnabled && VisualLoggerListener)
+	if (VisualLoggerSettings.bEnabled && VisualLoggerSettings.Listener)
 	{
-		const ULearningAgentsVisualLoggerObject* VisualLoggerObject = VisualLoggerListener->GetOrAddVisualLoggerObject(Tag);
+		const ULearningAgentsVisualLoggerObject* VisualLoggerObject = VisualLoggerSettings.Listener->GetOrAddVisualLoggerObject(Tag);
 
 		UE_LEARNING_AGENTS_VLOG_ANGLE_DEGREES(VisualLoggerObject, LogLearning, Display,
 			Angle,
 			RelativeAngle,
-			VisualLoggerLocation,
+			VisualLoggerSettings.Location,
 			10.0f,
-			VisualLoggerColor.ToFColor(true),
+			VisualLoggerSettings.Color.ToFColor(true),
 			TEXT(""));
 
-		UE_LEARNING_AGENTS_VLOG_STRING(VisualLoggerObject, LogLearning, Display, VisualLoggerLocation,
-			VisualLoggerColor.ToFColor(true),
+		UE_LEARNING_AGENTS_VLOG_STRING(VisualLoggerObject, LogLearning, Display, VisualLoggerSettings.Location,
+			VisualLoggerSettings.Color.ToFColor(true),
 			TEXT("Listener: %s\nTag: %s\nAgent Id: % 3i\nAngle: [% 6.1f]\nRelative Angle: [% 6.1f]\nLocal Angle: [% 6.1f]\nEncoded: [% 6.2f % 6.2f]"),
-			*VisualLoggerListener->GetName(),
+			*VisualLoggerSettings.Listener->GetName(),
 			*Tag.ToString(),
-			VisualLoggerAgentId,
+			VisualLoggerSettings.AgentId,
 			Angle,
 			RelativeAngle,
 			LocalAngle,
@@ -3085,11 +3026,7 @@ FLearningAgentsObservationObjectElement ULearningAgentsObservations::MakeAngleOb
 	const float RelativeAngle, 
 	const bool bActive,
 	const FName Tag,
-	const bool bVisualLoggerEnabled,
-	ULearningAgentsManagerListener* VisualLoggerListener,
-	const int32 VisualLoggerAgentId,
-	const FVector VisualLoggerLocation,
-	const FLinearColor VisualLoggerColor)
+	const FLearningAgentsObservationVisualLoggerSettings& VisualLoggerSettings)
 {
 	if (!bActive)
 	{
@@ -3102,25 +3039,16 @@ FLearningAgentsObservationObjectElement ULearningAgentsObservations::MakeAngleOb
 		FMath::RadiansToDegrees(RelativeAngle), 
 		bActive,
 		Tag, 
-		bVisualLoggerEnabled, 
-		VisualLoggerListener, 
-		VisualLoggerAgentId, 
-		VisualLoggerLocation, 
-		VisualLoggerColor);
+		VisualLoggerSettings);
 }
 
 FLearningAgentsObservationObjectElement ULearningAgentsObservations::MakeVelocityObservation(
-	ULearningAgentsObservationObject* Object, 
-	const FVector Velocity, 
-	const FTransform RelativeTransform, 
+	ULearningAgentsObservationObject* Object,
+	const FVector Velocity,
+	const FTransform RelativeTransform,
 	const bool bActive,
 	const FName Tag,
-	const bool bVisualLoggerEnabled,
-	ULearningAgentsManagerListener* VisualLoggerListener,
-	const int32 VisualLoggerAgentId,
-	const FVector VisualLoggerVelocityLocation,
-	const FVector VisualLoggerLocation,
-	const FLinearColor VisualLoggerColor)
+	const FLearningAgentsObservationVisualLoggerSettings& VisualLoggerSettings)
 {
 	if (!bActive)
 	{
@@ -3130,34 +3058,34 @@ FLearningAgentsObservationObjectElement ULearningAgentsObservations::MakeVelocit
 	const FVector LocalVelocity = RelativeTransform.InverseTransformVectorNoScale(Velocity);
 
 #if UE_LEARNING_AGENTS_ENABLE_VISUAL_LOG
-	if (bVisualLoggerEnabled && VisualLoggerListener)
+	if (VisualLoggerSettings.bEnabled && VisualLoggerSettings.Listener)
 	{
-		const ULearningAgentsVisualLoggerObject* VisualLoggerObject = VisualLoggerListener->GetOrAddVisualLoggerObject(Tag);
+		const ULearningAgentsVisualLoggerObject* VisualLoggerObject = VisualLoggerSettings.Listener->GetOrAddVisualLoggerObject(Tag);
 
 		UE_LEARNING_AGENTS_VLOG_ARROW(VisualLoggerObject, LogLearning, Display,
-			VisualLoggerVelocityLocation,
-			VisualLoggerVelocityLocation + Velocity,
-			VisualLoggerColor.ToFColor(true),
+			VisualLoggerSettings.DebugLocation,
+			VisualLoggerSettings.DebugLocation + Velocity,
+			VisualLoggerSettings.Color.ToFColor(true),
 			TEXT(""));
 
 		UE_LEARNING_AGENTS_VLOG_TRANSFORM(VisualLoggerObject, LogLearning, Display,
 			RelativeTransform.GetTranslation(),
 			RelativeTransform.GetRotation(),
-			VisualLoggerColor.ToFColor(true),
+			VisualLoggerSettings.Color.ToFColor(true),
 			TEXT(""));
 
-		UE_LEARNING_AGENTS_VLOG_STRING(VisualLoggerObject, LogLearning, Display, VisualLoggerLocation,
-			VisualLoggerColor.ToFColor(true),
+		UE_LEARNING_AGENTS_VLOG_STRING(VisualLoggerObject, LogLearning, Display, VisualLoggerSettings.Location,
+			VisualLoggerSettings.Color.ToFColor(true),
 			TEXT("Listener: %s\nTag: %s\nAgent Id: % 3i\nVelocity: [% 6.1f % 6.1f % 6.1f]\nLocal Velocity: [% 6.1f % 6.1f % 6.1f]"),
-			*VisualLoggerListener->GetName(),
+			*VisualLoggerSettings.Listener->GetName(),
 			*Tag.ToString(),
-			VisualLoggerAgentId,
+			VisualLoggerSettings.AgentId,
 			Velocity.X, Velocity.Y, Velocity.Z,
 			LocalVelocity.X, LocalVelocity.Y, LocalVelocity.Z);
 	}
 #endif
 
-	return MakeContinuousObservationFromArrayView(Object, { (float)LocalVelocity.X, (float)LocalVelocity.Y, (float)LocalVelocity.Z }, bActive, Tag);
+	return MakeContinuousObservationFromArrayView(Object, { (float)LocalVelocity.X, (float)LocalVelocity.Y, (float)LocalVelocity.Z }, bActive, Tag, FLearningAgentsObservationVisualLoggerSettings());
 }
 
 FLearningAgentsObservationObjectElement ULearningAgentsObservations::MakeDirectionObservation(
@@ -3166,13 +3094,7 @@ FLearningAgentsObservationObjectElement ULearningAgentsObservations::MakeDirecti
 	const FTransform RelativeTransform,
 	const bool bActive,
 	const FName Tag,
-	const bool bVisualLoggerEnabled,
-	ULearningAgentsManagerListener* VisualLoggerListener,
-	const int32 VisualLoggerAgentId,
-	const FVector VisualLoggerDirectionLocation,
-	const FVector VisualLoggerLocation,
-	const float VisualLoggerArrowLength,
-	const FLinearColor VisualLoggerColor)
+	const FLearningAgentsObservationVisualLoggerSettings& VisualLoggerSettings)
 {
 	if (!bActive)
 	{
@@ -3182,28 +3104,28 @@ FLearningAgentsObservationObjectElement ULearningAgentsObservations::MakeDirecti
 	const FVector LocalDirection = RelativeTransform.InverseTransformVectorNoScale(Direction).GetSafeNormal(UE_SMALL_NUMBER, FVector::ForwardVector);
 
 #if UE_LEARNING_AGENTS_ENABLE_VISUAL_LOG
-	if (bVisualLoggerEnabled && VisualLoggerListener)
+	if (VisualLoggerSettings.bEnabled && VisualLoggerSettings.Listener)
 	{
-		const ULearningAgentsVisualLoggerObject* VisualLoggerObject = VisualLoggerListener->GetOrAddVisualLoggerObject(Tag);
+		const ULearningAgentsVisualLoggerObject* VisualLoggerObject = VisualLoggerSettings.Listener->GetOrAddVisualLoggerObject(Tag);
 
 		UE_LEARNING_AGENTS_VLOG_ARROW(VisualLoggerObject, LogLearning, Display,
-			VisualLoggerDirectionLocation,
-			VisualLoggerDirectionLocation + VisualLoggerArrowLength * Direction,
-			VisualLoggerColor.ToFColor(true),
+			VisualLoggerSettings.DebugLocation,
+			VisualLoggerSettings.DebugLocation + VisualLoggerSettings.ArrowLength * Direction,
+			VisualLoggerSettings.Color.ToFColor(true),
 			TEXT(""));
 
 		UE_LEARNING_AGENTS_VLOG_TRANSFORM(VisualLoggerObject, LogLearning, Display,
 			RelativeTransform.GetTranslation(),
 			RelativeTransform.GetRotation(),
-			VisualLoggerColor.ToFColor(true),
+			VisualLoggerSettings.Color.ToFColor(true),
 			TEXT(""));
 
-		UE_LEARNING_AGENTS_VLOG_STRING(VisualLoggerObject, LogLearning, Display, VisualLoggerLocation,
-			VisualLoggerColor.ToFColor(true),
+		UE_LEARNING_AGENTS_VLOG_STRING(VisualLoggerObject, LogLearning, Display, VisualLoggerSettings.Location,
+			VisualLoggerSettings.Color.ToFColor(true),
 			TEXT("Listener: %s\nTag: %s\nAgent Id: % 3i\nDirection: [% 6.1f % 6.1f % 6.1f]\nLocal Direction: [% 6.1f % 6.1f % 6.1f]"),
-			*VisualLoggerListener->GetName(),
+			*VisualLoggerSettings.Listener->GetName(),
 			*Tag.ToString(),
-			VisualLoggerAgentId,
+			VisualLoggerSettings.AgentId,
 			Direction.X, Direction.Y, Direction.Z,
 			LocalDirection.X, LocalDirection.Y, LocalDirection.Z);
 	}
@@ -3223,11 +3145,7 @@ FLearningAgentsObservationObjectElement ULearningAgentsObservations::MakeLocatio
 	const FTransform RelativeTransform,
 	const bool bActive,
 	const FName Tag,
-	const bool bVisualLoggerEnabled,
-	ULearningAgentsManagerListener* VisualLoggerListener,
-	const int32 VisualLoggerAgentId,
-	const FVector VisualLoggerLocation,
-	const FLinearColor VisualLoggerColor)
+	const FLearningAgentsObservationVisualLoggerSettings& VisualLoggerSettings)
 {
 	if (!bActive)
 	{
@@ -3244,16 +3162,12 @@ FLearningAgentsObservationObjectElement ULearningAgentsObservations::MakeLocatio
 	const float LoopedDistance = SplineComponent->IsClosedLoop() ? FMath::Wrap(DistanceAlongSpline, 0.0f, TotalDistance) : FMath::Clamp(DistanceAlongSpline, 0.0f, TotalDistance);
 
 	return MakeLocationObservation(
-		Object, 
-		SplineComponent->GetLocationAtDistanceAlongSpline(LoopedDistance, ESplineCoordinateSpace::World), 
-		RelativeTransform, 
+		Object,
+		SplineComponent->GetLocationAtDistanceAlongSpline(LoopedDistance, ESplineCoordinateSpace::World),
+		RelativeTransform,
 		bActive,
-		Tag, 
-		bVisualLoggerEnabled, 
-		VisualLoggerListener, 
-		VisualLoggerAgentId, 
-		VisualLoggerLocation, 
-		VisualLoggerColor);
+		Tag,
+		VisualLoggerSettings);
 }
 
 FLearningAgentsObservationObjectElement ULearningAgentsObservations::MakeProportionAlongSplineObservation(
@@ -3262,11 +3176,7 @@ FLearningAgentsObservationObjectElement ULearningAgentsObservations::MakeProport
 	const float DistanceAlongSpline, 
 	const bool bActive,
 	const FName Tag,
-	const bool bVisualLoggerEnabled,
-	ULearningAgentsManagerListener* VisualLoggerListener,
-	const int32 VisualLoggerAgentId,
-	const FVector VisualLoggerLocation,
-	const FLinearColor VisualLoggerColor)
+	const FLearningAgentsObservationVisualLoggerSettings& VisualLoggerSettings)
 {
 	if (!bActive)
 	{
@@ -3291,12 +3201,8 @@ FLearningAgentsObservationObjectElement ULearningAgentsObservations::MakeProport
 			Angle,
 			0.0f,
 			bActive,
-			TEXT("AngleObservation"),
-			bVisualLoggerEnabled,
-			VisualLoggerListener,
-			VisualLoggerAgentId,
-			VisualLoggerLocation,
-			VisualLoggerColor), bActive, Tag);
+			TEXT("AngleObservation"), VisualLoggerSettings),
+			bActive, Tag);
 	}
 	else
 	{
@@ -3306,12 +3212,8 @@ FLearningAgentsObservationObjectElement ULearningAgentsObservations::MakeProport
 			Object, 
 			Proportion, 
 			bActive,
-			TEXT("FloatObservation"),
-			bVisualLoggerEnabled, 
-			VisualLoggerListener, 
-			VisualLoggerAgentId, 
-			VisualLoggerLocation, 
-			VisualLoggerColor), bActive, Tag);
+			TEXT("FloatObservation"), VisualLoggerSettings),
+			bActive, Tag);
 	}
 }
 
@@ -3322,12 +3224,7 @@ FLearningAgentsObservationObjectElement ULearningAgentsObservations::MakeDirecti
 	const FTransform RelativeTransform,
 	const bool bActive,
 	const FName Tag,
-	const bool bVisualLoggerEnabled,
-	ULearningAgentsManagerListener* VisualLoggerListener,
-	const int32 VisualLoggerAgentId,
-	const FVector VisualLoggerLocation,
-	const float VisualLoggerArrowLength,
-	const FLinearColor VisualLoggerColor)
+	const FLearningAgentsObservationVisualLoggerSettings& VisualLoggerSettings)
 {
 	if (!bActive)
 	{
@@ -3343,28 +3240,21 @@ FLearningAgentsObservationObjectElement ULearningAgentsObservations::MakeDirecti
 	const float TotalDistance = SplineComponent->GetSplineLength();
 	const float LoopedDistance = SplineComponent->IsClosedLoop() ? FMath::Wrap(DistanceAlongSpline, 0.0f, TotalDistance) : FMath::Clamp(DistanceAlongSpline, 0.0f, TotalDistance);
 
+	FLearningAgentsObservationVisualLoggerSettings DirectionVisualLoggerSettings = VisualLoggerSettings;
 #if UE_LEARNING_AGENTS_ENABLE_VISUAL_LOG
-	const FVector VisualLoggerDirectionLocation = 
-		(bVisualLoggerEnabled && VisualLoggerListener) ?
-		SplineComponent->GetLocationAtDistanceAlongSpline(LoopedDistance, ESplineCoordinateSpace::World) :
-		FVector::ZeroVector;
+	DirectionVisualLoggerSettings.DebugLocation = (VisualLoggerSettings.bEnabled && VisualLoggerSettings.Listener) ? SplineComponent->GetLocationAtDistanceAlongSpline(LoopedDistance, ESplineCoordinateSpace::World) : FVector::ZeroVector;
+
 #else
-	const FVector VisualLoggerDirectionLocation = FVector::ZeroVector;
+	DirectionVisualLoggerSettings.DebugLocation = FVector::ZeroVector;
 #endif
-
+	
 	return MakeDirectionObservation(
 		Object, 
 		SplineComponent->GetDirectionAtDistanceAlongSpline(LoopedDistance, ESplineCoordinateSpace::World), 
 		RelativeTransform, 
 		bActive,
-		Tag, 
-		bVisualLoggerEnabled, 
-		VisualLoggerListener, 
-		VisualLoggerAgentId, 
-		VisualLoggerDirectionLocation, 
-		VisualLoggerLocation, 
-		VisualLoggerArrowLength, 
-		VisualLoggerColor);
+		Tag,
+		DirectionVisualLoggerSettings);
 }
 
 FLearningAgentsObservationObjectElement ULearningAgentsObservations::MakeProportionAlongRayObservation(
@@ -3375,11 +3265,7 @@ FLearningAgentsObservationObjectElement ULearningAgentsObservations::MakeProport
 	const ECollisionChannel CollisionChannel,
 	const bool bActive,
 	const FName Tag,
-	const bool bVisualLoggerEnabled,
-	ULearningAgentsManagerListener* VisualLoggerListener,
-	const int32 VisualLoggerAgentId,
-	const FVector VisualLoggerLocation,
-	const FLinearColor VisualLoggerColor)
+	const FLearningAgentsObservationVisualLoggerSettings& VisualLoggerSettings)
 {
 	if (!bActive)
 	{
@@ -3399,14 +3285,14 @@ FLearningAgentsObservationObjectElement ULearningAgentsObservations::MakeProport
 	const float Encoded = bHit ? 1.0f - TraceHit.Time : 0.0f;
 
 #if UE_LEARNING_AGENTS_ENABLE_VISUAL_LOG
-	if (bVisualLoggerEnabled && VisualLoggerListener)
+	if (VisualLoggerSettings.bEnabled && VisualLoggerSettings.Listener)
 	{
-		const ULearningAgentsVisualLoggerObject* VisualLoggerObject = VisualLoggerListener->GetOrAddVisualLoggerObject(Tag);
+		const ULearningAgentsVisualLoggerObject* VisualLoggerObject = VisualLoggerSettings.Listener->GetOrAddVisualLoggerObject(Tag);
 
 		UE_LEARNING_AGENTS_VLOG_SEGMENT(VisualLoggerObject, LogLearning, Display,
 			RayStartWorld,
 			RayEndWorld,
-			VisualLoggerColor.ToFColor(true),
+			VisualLoggerSettings.Color.ToFColor(true),
 			TEXT(""));
 
 		if (bHit)
@@ -3414,22 +3300,22 @@ FLearningAgentsObservationObjectElement ULearningAgentsObservations::MakeProport
 			UE_LEARNING_AGENTS_VLOG_LOCATION(VisualLoggerObject, LogLearning, Display,
 				RayStartWorld + TraceHit.Time * (RayEndWorld - RayStartWorld),
 				10.f,
-				VisualLoggerColor.ToFColor(true),
+				VisualLoggerSettings.Color.ToFColor(true),
 				TEXT(""));
 		}
 
 		UE_LEARNING_AGENTS_VLOG_TRANSFORM(VisualLoggerObject, LogLearning, Display,
 			RayTransform.GetTranslation(),
 			RayTransform.GetRotation(),
-			VisualLoggerColor.ToFColor(true),
+			VisualLoggerSettings.Color.ToFColor(true),
 			TEXT(""));
 
-		UE_LEARNING_AGENTS_VLOG_STRING(VisualLoggerObject, LogLearning, Display, VisualLoggerLocation,
-			VisualLoggerColor.ToFColor(true),
+		UE_LEARNING_AGENTS_VLOG_STRING(VisualLoggerObject, LogLearning, Display, VisualLoggerSettings.Location,
+			VisualLoggerSettings.Color.ToFColor(true),
 			TEXT("Listener: %s\nTag: %s\nAgent Id: % 3i\nRay Start: [% 6.1f % 6.1f % 6.1f]\nRay End: [% 6.1f % 6.1f % 6.1f]\nRay Start World: [% 6.1f % 6.1f % 6.1f]\nRay End World: [% 6.1f % 6.1f % 6.1f]\nCollision Channel: [%s]\nHit: [%s]\nEncoded: [% 6.2f]"),
-			*VisualLoggerListener->GetName(),
+			*VisualLoggerSettings.Listener->GetName(),
 			*Tag.ToString(),
-			VisualLoggerAgentId,
+			VisualLoggerSettings.AgentId,
 			RayStart.X, RayStart.Y, RayStart.Z,
 			RayEnd.X, RayEnd.Y, RayEnd.Z,
 			RayStartWorld.X, RayStartWorld.Y, RayStartWorld.Z,
@@ -3452,11 +3338,7 @@ FLearningAgentsObservationObjectElement ULearningAgentsObservations::MakeLidarOb
 	const AActor* ActorToIgnore,
 	const bool bActive,
 	const FName Tag,
-	const bool bVisualLoggerEnabled,
-	ULearningAgentsManagerListener* VisualLoggerListener,
-	const int32 VisualLoggerAgentId,
-	const FVector VisualLoggerLocation,
-	const FLinearColor VisualLoggerColor)
+	const FLearningAgentsObservationVisualLoggerSettings& VisualLoggerSettings)
 {
 	if (!bActive)
 	{
@@ -3485,14 +3367,14 @@ FLearningAgentsObservationObjectElement ULearningAgentsObservations::MakeLidarOb
 	const float Encoded = bHit ? FMath::Clamp(1.0f - TraceHit.Distance / RayRange, 0.0f, 1.0f) : 0.0f;
 
 #if UE_LEARNING_AGENTS_ENABLE_VISUAL_LOG
-	if (bVisualLoggerEnabled && VisualLoggerListener)
+	if (VisualLoggerSettings.bEnabled && VisualLoggerSettings.Listener)
 	{
-		const ULearningAgentsVisualLoggerObject* VisualLoggerObject = VisualLoggerListener->GetOrAddVisualLoggerObject(Tag);
+		const ULearningAgentsVisualLoggerObject* VisualLoggerObject = VisualLoggerSettings.Listener->GetOrAddVisualLoggerObject(Tag);
 
 		UE_LEARNING_AGENTS_VLOG_SEGMENT(VisualLoggerObject, LogLearning, Display,
 			RayStartWorld,
 			RayEndWorld,
-			VisualLoggerColor.ToFColor(true),
+			VisualLoggerSettings.Color.ToFColor(true),
 			TEXT(""));
 
 		if (bHit)
@@ -3500,22 +3382,22 @@ FLearningAgentsObservationObjectElement ULearningAgentsObservations::MakeLidarOb
 			UE_LEARNING_AGENTS_VLOG_LOCATION(VisualLoggerObject, LogLearning, Display,
 				TraceHit.Location,
 				10.f,
-				VisualLoggerColor.ToFColor(true),
+				VisualLoggerSettings.Color.ToFColor(true),
 				TEXT(""));
 		}
 
 		UE_LEARNING_AGENTS_VLOG_TRANSFORM(VisualLoggerObject, LogLearning, Display,
 			RayTransform.GetTranslation(),
 			RayTransform.GetRotation(),
-			VisualLoggerColor.ToFColor(true),
+			VisualLoggerSettings.Color.ToFColor(true),
 			TEXT(""));
 
-		UE_LEARNING_AGENTS_VLOG_STRING(VisualLoggerObject, LogLearning, Display, VisualLoggerLocation,
-			VisualLoggerColor.ToFColor(true),
+		UE_LEARNING_AGENTS_VLOG_STRING(VisualLoggerObject, LogLearning, Display, VisualLoggerSettings.Location,
+			VisualLoggerSettings.Color.ToFColor(true),
 			TEXT("Listener: %s\nTag: %s\nAgent Id: % 3i\nRay Start: [% 6.1f % 6.1f % 6.1f]\nRay End: [% 6.1f % 6.1f % 6.1f]\nRay Start World: [% 6.1f % 6.1f % 6.1f]\nRay End World: [% 6.1f % 6.1f % 6.1f]\nCollision Channel: [%s]\nHit: [%s]\nEncoded: [% 6.2f]"),
-			*VisualLoggerListener->GetName(),
+			*VisualLoggerSettings.Listener->GetName(),
 			*Tag.ToString(),
-			VisualLoggerAgentId,
+			VisualLoggerSettings.AgentId,
 			RayStart.X, RayStart.Y, RayStart.Z,
 			RayEnd.X, RayEnd.Y, RayEnd.Z,
 			RayStartWorld.X, RayStartWorld.Y, RayStartWorld.Z,
diff --git a/Engine/Plugins/Experimental/LearningAgents/Source/LearningAgents/Private/LearningAgentsPolicy.cpp b/Engine/Plugins/Experimental/LearningAgents/Source/LearningAgents/Private/LearningAgentsPolicy.cpp
--- a/Engine/Plugins/Experimental/LearningAgents/Source/LearningAgents/Private/LearningAgentsPolicy.cpp
+++ b/Engine/Plugins/Experimental/LearningAgents/Source/LearningAgents/Private/LearningAgentsPolicy.cpp
@@ -19,23 +19,6 @@
 
 #include UE_INLINE_GENERATED_CPP_BY_NAME(LearningAgentsPolicy)
 
-namespace UE::Learning::Agents::Policy::Private
-{
-	static inline UE::NNE::RuntimeBasic::FModelBuilder::EActivationFunction GetBuilderActivationFunction(const ELearningAgentsActivationFunction ActivationFunction)
-	{
-		switch (ActivationFunction)
-		{
-		case ELearningAgentsActivationFunction::ReLU: return UE::NNE::RuntimeBasic::FModelBuilder::EActivationFunction::ReLU;
-		case ELearningAgentsActivationFunction::ELU: return UE::NNE::RuntimeBasic::FModelBuilder::EActivationFunction::ELU;
-		case ELearningAgentsActivationFunction::TanH: return UE::NNE::RuntimeBasic::FModelBuilder::EActivationFunction::TanH;
-		case ELearningAgentsActivationFunction::GELU: return UE::NNE::RuntimeBasic::FModelBuilder::EActivationFunction::GELU;
-		default:
-			UE_LOG(LogLearning, Error, TEXT("Unknown Activation Function"));
-			return UE::NNE::RuntimeBasic::FModelBuilder::EActivationFunction::ReLU;
-		}
-	}
-}
-
 ULearningAgentsPolicy::ULearningAgentsPolicy() : Super(FObjectInitializer::Get()) {}
 ULearningAgentsPolicy::ULearningAgentsPolicy(FVTableHelper& Helper) : Super(Helper) {}
 ULearningAgentsPolicy::~ULearningAgentsPolicy() = default;
@@ -283,7 +266,7 @@ void ULearningAgentsPolicy::SetupPolicy(
 						ActionEncodedVectorSize,
 						PolicyHiddenLayerSize,
 						PolicyHiddenLayerNum + 2, // Add 2 to account for input and output layers
-						UE::Learning::Agents::Policy::Private::GetBuilderActivationFunction(PolicyActivationFunction)),
+						UE::Learning::Agents::NeuralNetwork::GetBuilderActivationFunction(PolicyActivationFunction)),
 					Builder.MakeDenormalize(
 						ActionEncodedVectorSize,
 						Builder.MakeValuesZero(ActionEncodedVectorSize),
@@ -298,7 +281,7 @@ void ULearningAgentsPolicy::SetupPolicy(
 						PolicyHiddenLayerSize,
 						PolicyHiddenLayerSize,
 						PolicyHiddenLayerNum / 2 + 2, // Add 2 to account for input and output layers
-						UE::Learning::Agents::Policy::Private::GetBuilderActivationFunction(PolicyActivationFunction),
+						UE::Learning::Agents::NeuralNetwork::GetBuilderActivationFunction(PolicyActivationFunction),
 						true),
 					Builder.MakeMemoryCellLayer(
 						PolicyHiddenLayerSize,
@@ -311,7 +294,7 @@ void ULearningAgentsPolicy::SetupPolicy(
 							ActionEncodedVectorSize,
 							PolicyHiddenLayerSize,
 							PolicyHiddenLayerNum / 2 + 2, // Add 2 to account for input and output layers
-							UE::Learning::Agents::Policy::Private::GetBuilderActivationFunction(PolicyActivationFunction)),
+							UE::Learning::Agents::NeuralNetwork::GetBuilderActivationFunction(PolicyActivationFunction)),
 						Builder.MakeDenormalize(
 							ActionEncodedVectorSize,
 							Builder.MakeValuesZero(ActionEncodedVectorSize),
@@ -328,7 +311,7 @@ void ULearningAgentsPolicy::SetupPolicy(
 						PolicyHiddenLayerSize,
 						PolicyHiddenLayerSize,
 						PolicyHiddenLayerNum / 2 + 2, // Add 2 to account for input and output layers
-						UE::Learning::Agents::Policy::Private::GetBuilderActivationFunction(PolicyActivationFunction),
+						UE::Learning::Agents::NeuralNetwork::GetBuilderActivationFunction(PolicyActivationFunction),
 						true),
 					Builder.MakeGRUCellLayer(
 						PolicyHiddenLayerSize,
@@ -340,7 +323,7 @@ void ULearningAgentsPolicy::SetupPolicy(
 							ActionEncodedVectorSize,
 							PolicyHiddenLayerSize,
 							PolicyHiddenLayerNum / 2 + 2, // Add 2 to account for input and output layers
-							UE::Learning::Agents::Policy::Private::GetBuilderActivationFunction(PolicyActivationFunction)),
+							UE::Learning::Agents::NeuralNetwork::GetBuilderActivationFunction(PolicyActivationFunction)),
 						Builder.MakeDenormalize(
 							ActionEncodedVectorSize,
 							Builder.MakeValuesZero(ActionEncodedVectorSize),
