[LearningAgents] Flow Matching and Normalizations

diff --git a/Engine/Plugins/Experimental/LearningAgents/Source/LearningAgents/Public/LearningAgentsActions.h b/Engine/Plugins/Experimental/LearningAgents/Source/LearningAgents/Public/LearningAgentsActions.h
--- a/Engine/Plugins/Experimental/LearningAgents/Source/LearningAgents/Public/LearningAgentsActions.h
+++ b/Engine/Plugins/Experimental/LearningAgents/Source/LearningAgents/Public/LearningAgentsActions.h
@@ -49,6 +49,100 @@ struct FLearningAgentsActionModifierElement
 	UE::Learning::Action::FModifierElement ModifierElement;
 };
 
+/** Settings for normalizing observations */
+USTRUCT(BlueprintType)
+struct FLearningAgentsActionNormalizationSettings
+{
+	GENERATED_BODY()
+
+	/** The normalization method to use. */
+	UPROPERTY(EditAnywhere, BlueprintReadWrite, Category = "LearningAgents")
+	ELearningAgentsNormalization NormalizationOption = ELearningAgentsNormalization::Manual;
+
+	/** Per-element standard deviation values used for normalization. If empty, no scaling is applied. */
+	UPROPERTY(EditAnywhere, BlueprintReadWrite, Category = "LearningAgents", meta = (EditCondition = "NormalizationOption == ELearningAgentsNormalization::Manual"))
+	TArray<float> Scales;
+
+	/** Per-element mean values used for normalization. If empty, no bias is applied. */
+	UPROPERTY(EditAnywhere, BlueprintReadWrite, Category = "LearningAgents", meta = (EditCondition = "NormalizationOption == ELearningAgentsNormalization::Manual"))
+	TArray<float> Offsets;
+};
+
+/** Settings for normalizing observations */
+USTRUCT(BlueprintType)
+struct FLearningAgentsFloatActionNormalizationSettings
+{
+	GENERATED_BODY()
+
+	/** The normalization method to use. */
+	UPROPERTY(EditAnywhere, BlueprintReadWrite, Category = "LearningAgents")
+	ELearningAgentsNormalization NormalizationOption = ELearningAgentsNormalization::Manual;
+
+	/** Per-element standard deviation values used for normalization. If empty, no scaling is applied. */
+	UPROPERTY(EditAnywhere, BlueprintReadWrite, Category = "LearningAgents", meta = (EditCondition = "NormalizationOption == ELearningAgentsNormalization::Manual"))
+	float Scale = 1.0f;
+
+	/** Per-element mean values used for normalization. If empty, no bias is applied. */
+	UPROPERTY(EditAnywhere, BlueprintReadWrite, Category = "LearningAgents", meta = (EditCondition = "NormalizationOption == ELearningAgentsNormalization::Manual"))
+	float Offset = 0.0f;
+};
+
+/** Settings for normalizing observations */
+USTRUCT(BlueprintType)
+struct FLearningAgentsVectorActionNormalizationSettings
+{
+	GENERATED_BODY()
+
+	/** The normalization method to use. */
+	UPROPERTY(EditAnywhere, BlueprintReadWrite, Category = "LearningAgents")
+	ELearningAgentsNormalization NormalizationOption = ELearningAgentsNormalization::Manual;
+
+	/** Per-element standard deviation values used for normalization. If empty, no scaling is applied. */
+	UPROPERTY(EditAnywhere, BlueprintReadWrite, Category = "LearningAgents", meta = (EditCondition = "NormalizationOption == ELearningAgentsNormalization::Manual"))
+	FVector Scales = FVector::OneVector;
+
+	/** Per-element mean values used for normalization. If empty, no bias is applied. */
+	UPROPERTY(EditAnywhere, BlueprintReadWrite, Category = "LearningAgents", meta = (EditCondition = "NormalizationOption == ELearningAgentsNormalization::Manual"))
+	FVector Offsets = FVector::ZeroVector;
+};
+
+class ULearningAgentsManagerListener;
+
+/** Settings for visual logger output */
+USTRUCT(BlueprintType)
+struct FLearningAgentsVisualLoggerSettings
+{
+	GENERATED_BODY()
+
+	/** Whether visual logging is enabled. */
+	UPROPERTY(EditAnywhere, BlueprintReadWrite, Category = "LearningAgents")
+	bool bEnabled = false;
+
+	/** The listener object which is making this action. This must be set to use logging. */
+	UPROPERTY(EditAnywhere, BlueprintReadWrite, Category = "LearningAgents")
+	TObjectPtr<ULearningAgentsManagerListener> Listener = nullptr;
+
+	/** The agent id associated with this action. */
+	UPROPERTY(EditAnywhere, BlueprintReadWrite, Category = "LearningAgents")
+	int32 AgentId = -1;
+
+	/** An additional location for visualizing a specific action element (e.g., rotation arrow origin, velocity arrow start). */
+	UPROPERTY(EditAnywhere, BlueprintReadWrite, Category = "LearningAgents")
+	FVector DebugLocation = FVector::ZeroVector;
+
+	/** A location for the visual logger text information in the world. */
+	UPROPERTY(EditAnywhere, BlueprintReadWrite, Category = "LearningAgents")
+	FVector Location = FVector::ZeroVector;
+
+	/** The color for the visual logger display. */
+	UPROPERTY(EditAnywhere, BlueprintReadWrite, Category = "LearningAgents")
+	FLinearColor Color = FLinearColor::Blue;
+
+	/** The length of debug arrows to display for direction/velocity actions. */
+	UPROPERTY(EditAnywhere, BlueprintReadWrite, Category = "LearningAgents")
+	float ArrowLength = 100.0f;
+};
+
 /** Comparison and Hashing operators for Action Elements */
 
 LEARNINGAGENTS_API bool operator==(const FLearningAgentsActionSchemaElement& Lhs, const FLearningAgentsActionSchemaElement& Rhs);
@@ -212,16 +306,16 @@ public:
 	 *
 	 * @param Schema The Action Schema
 	 * @param Size The number of float values in the action.
-	 * @param Scale The scale used for this action.
-	 * @param Tag The tag of this new action. Used during action object validation and debugging.
+	 * @param Normalization The normalization used for this action.
 	 * @param bActive Whether to enable the specified action as a valid network output. If set to inactive, this action will be considered an empty action. This flag allows easy toggles of specified actions.
+	 * @param Tag The tag of this new action. Used during action object validation and debugging.
 	 * @return The newly created action schema element.
 	 */
-	UFUNCTION(BlueprintPure, Category = "LearningAgents", meta = (AdvancedDisplay = 3))
-	static UE_API FLearningAgentsActionSchemaElement SpecifyContinuousAction(ULearningAgentsActionSchema* Schema, const int32 Size, const float Scale = 1.0f, const bool bActive = true, const FName Tag = TEXT("ContinuousAction"));
+	UFUNCTION(BlueprintPure, Category = "LearningAgents", meta = (AdvancedDisplay = 3, AutoCreateRefTerm = "NormalizationSettings"))
+	static UE_API FLearningAgentsActionSchemaElement SpecifyContinuousAction(ULearningAgentsActionSchema* Schema, const int32 Size, const FLearningAgentsActionNormalizationSettings& NormalizationSettings = FLearningAgentsActionNormalizationSettings(), const bool bActive = true, const FName Tag = TEXT("ContinuousAction"));
 
 	/**
-	 * Specifies a new exclusive discrete action. This represents an action which is an exclusive choice from a number of discrete options, sampled 
+	 * Specifies a new exclusive discrete action. This represents an action which is an exclusive choice from a number of discrete options, sampled
 	 * from a Categorical distribution.
 	 *
 	 * @param Schema The Action Schema
@@ -612,41 +706,41 @@ public:
 	static UE_API FLearningAgentsActionSchemaElement SpecifyBoolAction(ULearningAgentsActionSchema* Schema, const float PriorProbability = 0.5f, const bool bActive = true, const FName Tag = TEXT("BoolAction"));
 
 	/**
-	 * Specifies a new float action. This represents an action which is a single float sampled from a Gaussian distribution. It can be used as a 
+	 * Specifies a new float action. This represents an action which is a single float sampled from a Gaussian distribution. It can be used as a
 	 * catch-all for situations where a type-specific action does not exist.
 	 *
 	 * @param Schema The Action Schema
-	 * @param FloatScale The scale used for this action.
+	 * @param NormalizationSettings The normalization used for this action.
 	 * @param bActive Whether to enable the specified action as a valid network output. If set to inactive, this action will be considered an empty action. This flag allows easy toggles of specified actions.
 	 * @param Tag The tag of this new action. Used during action object validation and debugging.
 	 * @return The newly created action schema element.
 	 */
-	UFUNCTION(BlueprintPure, Category = "LearningAgents", meta = (AdvancedDisplay = 2))
-	static UE_API FLearningAgentsActionSchemaElement SpecifyFloatAction(ULearningAgentsActionSchema* Schema, const float FloatScale = 1.0f, const bool bActive = true, const FName Tag = TEXT("FloatAction"));
+	UFUNCTION(BlueprintPure, Category = "LearningAgents", meta = (AdvancedDisplay = 2, AutoCreateRefTerm = "NormalizationSettings"))
+	static UE_API FLearningAgentsActionSchemaElement SpecifyFloatAction(ULearningAgentsActionSchema* Schema, const FLearningAgentsFloatActionNormalizationSettings& NormalizationSettings = FLearningAgentsFloatActionNormalizationSettings(), const bool bActive = true, const FName Tag = TEXT("FloatAction"));
 
 	/**
 	 * Specifies a new location action. This represents an action which is a location sampled from a Gaussian distribution.
 	 *
 	 * @param Schema The Action Schema
-	 * @param LocationScale The scale used for this action in cm.
+	 * @param NormalizationSettings The normalization used for this action in cm.
 	 * @param bActive Whether to enable the specified action as a valid network output. If set to inactive, this action will be considered an empty action. This flag allows easy toggles of specified actions.
 	 * @param Tag The tag of this new action. Used during action object validation and debugging.
 	 * @return The newly created action schema element.
 	 */
-	UFUNCTION(BlueprintPure, Category = "LearningAgents", meta = (AdvancedDisplay = 2))
-	static UE_API FLearningAgentsActionSchemaElement SpecifyLocationAction(ULearningAgentsActionSchema* Schema, const float LocationScale = 100.0f, const bool bActive = true, const FName Tag = TEXT("LocationAction"));
+	UFUNCTION(BlueprintPure, Category = "LearningAgents", meta = (AdvancedDisplay = 2, AutoCreateRefTerm = "NormalizationSettings", NormalizationSettings = "(Scales=(X=100,Y=100,Z=100))"))
+	static UE_API FLearningAgentsActionSchemaElement SpecifyLocationAction(ULearningAgentsActionSchema* Schema, const FLearningAgentsVectorActionNormalizationSettings& NormalizationSettings = FLearningAgentsVectorActionNormalizationSettings(), const bool bActive = true, const FName Tag = TEXT("LocationAction"));
 
 	/**
 	 * Specifies a new rotation action. This represents an action which is a rotation sampled from a Gaussian distribution in the angle-axis space.
 	 *
 	 * @param Schema The Action Schema
-	 * @param RotationScale The scale used for this action in degrees.
+	 * @param NormalizationSettings The NormalizationSettings used for this action in degrees.
 	 * @param bActive Whether to enable the specified action as a valid network output. If set to inactive, this action will be considered an empty action. This flag allows easy toggles of specified actions.
 	 * @param Tag The tag of this new action. Used during action object validation and debugging.
 	 * @return The newly created action schema element.
 	 */
-	UFUNCTION(BlueprintPure, Category = "LearningAgents", meta = (AdvancedDisplay = 2))
-	static UE_API FLearningAgentsActionSchemaElement SpecifyRotationAction(ULearningAgentsActionSchema* Schema, const float RotationScale = 90.0f, const bool bActive = true, const FName Tag = TEXT("RotationAction"));
+	UFUNCTION(BlueprintPure, Category = "LearningAgents", meta = (AdvancedDisplay = 2, AutoCreateRefTerm = "NormalizationSettings", NormalizationSettings = "(Scales=(X=90,Y=90,Z=90))"))
+	static UE_API FLearningAgentsActionSchemaElement SpecifyRotationAction(ULearningAgentsActionSchema* Schema, const FLearningAgentsVectorActionNormalizationSettings& NormalizationSettings = FLearningAgentsVectorActionNormalizationSettings(), const bool bActive = true, const FName Tag = TEXT("RotationAction"));
 
 	/**
 	 * Specifies a new scale action. This represents an action which is a scale sampled from a Gaussian distribution in the log space.
@@ -657,26 +751,26 @@ public:
 	 * @param Tag The tag of this new action. Used during action object validation and debugging.
 	 * @return The newly created action schema element.
 	 */
-	UFUNCTION(BlueprintPure, Category = "LearningAgents", meta = (AdvancedDisplay = 2))
-	static UE_API FLearningAgentsActionSchemaElement SpecifyScaleAction(ULearningAgentsActionSchema* Schema, const float ScaleScale = 1.0f, const bool bActive = true, const FName Tag = TEXT("ScaleAction"));
+	UFUNCTION(BlueprintPure, Category = "LearningAgents", meta = (AdvancedDisplay = 2, AutoCreateRefTerm = "NormalizationSettings"))
+	static UE_API FLearningAgentsActionSchemaElement SpecifyScaleAction(ULearningAgentsActionSchema* Schema, const FLearningAgentsVectorActionNormalizationSettings& NormalizationSettings = FLearningAgentsVectorActionNormalizationSettings(), const bool bActive = true, const FName Tag = TEXT("ScaleAction"));
 
 	/**
 	 * Specifies a new transform action.
 	 *
 	 * @param Schema The Action Schema
-	 * @param LocationScale The scale used for the Location part of the transform in this action in cm.
-	 * @param RotationScale The scale used for the Rotation part of the transform in this action in degrees.
-	 * @param ScaleScale The scale used for the Scale part of the transform in this action.
+	 * @param LocationNormalization The normalization used for the Location part of the transform in this action in cm.
+	 * @param RotationNormalization The normalization used for the Rotation part of the transform in this action in degrees.
+	 * @param ScaleNormalization The normalization used for the Scale part of the transform in this action.
 	 * @param bActive Whether to enable the specified action as a valid network output. If set to inactive, this action will be considered an empty action. This flag allows easy toggles of specified actions.
 	 * @param Tag The tag of this new action. Used during action object validation and debugging.
 	 * @return The newly created action schema element.
 	 */
-	UFUNCTION(BlueprintPure, Category = "LearningAgents", meta = (AdvancedDisplay = 4))
+	UFUNCTION(BlueprintPure, Category = "LearningAgents", meta = (AdvancedDisplay = 4, AutoCreateRefTerm = "LocationNormalization, RotationNormalization, ScaleNormalization", LocationNormalization = "(Scales=(X=100,Y=100,Z=100))", RotationNormalization = "(Scales=(X=90,Y=90,Z=90))"))
 	static UE_API FLearningAgentsActionSchemaElement SpecifyTransformAction(
 		ULearningAgentsActionSchema* Schema, 
-		const float LocationScale = 100.0f, 
-		const float RotationScale = 90.0f, 
-		const float ScaleScale = 1.0f, 
+		const FLearningAgentsVectorActionNormalizationSettings& LocationNormalization = FLearningAgentsVectorActionNormalizationSettings(),
+		const FLearningAgentsVectorActionNormalizationSettings& RotationNormalization = FLearningAgentsVectorActionNormalizationSettings(),
+		const FLearningAgentsVectorActionNormalizationSettings& ScaleNormalization = FLearningAgentsVectorActionNormalizationSettings(),
 		const bool bActive = true, 
 		const FName Tag = TEXT("TransformAction"));
 
@@ -684,25 +778,25 @@ public:
 	 * Specifies a new angle action. This represents an action which is an angle sampled from a Gaussian distribution centered around zero.
 	 *
 	 * @param Schema The Action Schema
-	 * @param AngleScale The scale used for this action in degrees.
+	 * @param NormalizationSettings The normalization used for this action in degrees.
 	 * @param bActive Whether to enable the specified action as a valid network output. If set to inactive, this action will be considered an empty action. This flag allows easy toggles of specified actions.
 	 * @param Tag The tag of this new action. Used during action object validation and debugging.
 	 * @return The newly created action schema element.
 	 */
-	UFUNCTION(BlueprintPure, Category = "LearningAgents", meta = (AdvancedDisplay = 2))
-	static UE_API FLearningAgentsActionSchemaElement SpecifyAngleAction(ULearningAgentsActionSchema* Schema, const float AngleScale = 90.0f, const bool bActive = true, const FName Tag = TEXT("AngleAction"));
+	UFUNCTION(BlueprintPure, Category = "LearningAgents", meta = (AdvancedDisplay = 2, AutoCreateRefTerm = "NormalizationSettings", NormalizationSettings = "(Scale=90)"))
+	static UE_API FLearningAgentsActionSchemaElement SpecifyAngleAction(ULearningAgentsActionSchema* Schema, const FLearningAgentsFloatActionNormalizationSettings& NormalizationSettings = FLearningAgentsFloatActionNormalizationSettings(), const bool bActive = true, const FName Tag = TEXT("AngleAction"));
 
 	/**
 	 * Specifies a new velocity action. This represents an action which is a velocity sampled from a Gaussian distribution.
 	 *
 	 * @param Schema The Action Schema
-	 * @param VelocityScale The scale used for this action in cm/s.
+	 * @param NormalizationSettings The normalization used for this action in cm/s.
 	 * @param bActive Whether to enable the specified action as a valid network output. If set to inactive, this action will be considered an empty action. This flag allows easy toggles of specified actions.
 	 * @param Tag The tag of this new action. Used during action object validation and debugging.
 	 * @return The newly created action schema element.
 	 */
-	UFUNCTION(BlueprintPure, Category = "LearningAgents", meta = (AdvancedDisplay = 2))
-	static UE_API FLearningAgentsActionSchemaElement SpecifyVelocityAction(ULearningAgentsActionSchema* Schema, const float VelocityScale = 200.0f, const bool bActive = true, const FName Tag = TEXT("VelocityAction"));
+	UFUNCTION(BlueprintPure, Category = "LearningAgents", meta = (AdvancedDisplay = 2, AutoCreateRefTerm = "NormalizationSettings", NormalizationSettings = "(Scales=(X=200,Y=200,Z=200))"))
+	static UE_API FLearningAgentsActionSchemaElement SpecifyVelocityAction(ULearningAgentsActionSchema* Schema, const FLearningAgentsVectorActionNormalizationSettings& NormalizationSettings = FLearningAgentsVectorActionNormalizationSettings(), const bool bActive = true, const FName Tag = TEXT("VelocityAction"));
 
 	/**
 	 * Specifies a new direction action. This represents an action which is a direction sampled from a Gaussian distribution and normalized.
@@ -720,98 +814,66 @@ public:
 	UFUNCTION(BlueprintPure, Category = "LearningAgents", meta = (AdvancedDisplay = 1))
 	static UE_API FLearningAgentsActionObjectElement MakeNullAction(ULearningAgentsActionObject* Object, const FName Tag = TEXT("NullAction"));
 
-	UFUNCTION(BlueprintPure, Category = "LearningAgents", meta = (AdvancedDisplay = 2, DefaultToSelf = "VisualLoggerListener"))
+	UFUNCTION(BlueprintPure, Category = "LearningAgents", meta = (AdvancedDisplay = 2, AutoCreateRefTerm = "VisualLoggerSettings"))
 	static UE_API FLearningAgentsActionObjectElement MakeContinuousAction(
 		ULearningAgentsActionObject* Object,
 		const TArray<float>& Values,
-		const bool bActive = true, 
+		const bool bActive = true,
 		const FName Tag = TEXT("ContinuousAction"),
-		const bool bVisualLoggerEnabled = false,
-		ULearningAgentsManagerListener* VisualLoggerListener = nullptr,
-		const int32 VisualLoggerAgentId = -1,
-		const FVector VisualLoggerLocation = FVector::ZeroVector,
-		const FLinearColor VisualLoggerColor = FLinearColor::Blue);
+		const FLearningAgentsVisualLoggerSettings& VisualLoggerSettings = FLearningAgentsVisualLoggerSettings());
 
 	static UE_API FLearningAgentsActionObjectElement MakeContinuousActionFromArrayView(
 		ULearningAgentsActionObject* Object,
 		const TArrayView<const float> Values,
-		const bool bActive = true, 
+		const bool bActive = true,
 		const FName Tag = TEXT("ContinuousAction"),
-		const bool bVisualLoggerEnabled = false,
-		ULearningAgentsManagerListener* VisualLoggerListener = nullptr,
-		const int32 VisualLoggerAgentId = -1,
-		const FVector VisualLoggerLocation = FVector::ZeroVector,
-		const FLinearColor VisualLoggerColor = FLinearColor::Blue);
+		const FLearningAgentsVisualLoggerSettings& VisualLoggerSettings = FLearningAgentsVisualLoggerSettings());
 
-	UFUNCTION(BlueprintPure, Category = "LearningAgents", meta = (AdvancedDisplay = 2, DefaultToSelf = "VisualLoggerListener"))
+	UFUNCTION(BlueprintPure, Category = "LearningAgents", meta = (AdvancedDisplay = 2, AutoCreateRefTerm = "VisualLoggerSettings"))
 	static UE_API FLearningAgentsActionObjectElement MakeExclusiveDiscreteAction(
 		ULearningAgentsActionObject* Object,
 		const int32 Index,
 		const bool bActive = true, 
 		const FName Tag = TEXT("DiscreteExclusiveAction"),
-		const bool bVisualLoggerEnabled = false,
-		ULearningAgentsManagerListener* VisualLoggerListener = nullptr,
-		const int32 VisualLoggerAgentId = -1,
-		const FVector VisualLoggerLocation = FVector::ZeroVector,
-		const FLinearColor VisualLoggerColor = FLinearColor::Blue);
+		const FLearningAgentsVisualLoggerSettings& VisualLoggerSettings = FLearningAgentsVisualLoggerSettings());
 
-	UFUNCTION(BlueprintPure, Category = "LearningAgents", meta = (AdvancedDisplay = 2, DefaultToSelf = "VisualLoggerListener"))
+	UFUNCTION(BlueprintPure, Category = "LearningAgents", meta = (AdvancedDisplay = 2, AutoCreateRefTerm = "VisualLoggerSettings"))
 	static UE_API FLearningAgentsActionObjectElement MakeNamedExclusiveDiscreteAction(
 		ULearningAgentsActionObject* Object,
 		const FName Name,
 		const bool bActive = true, 
 		const FName Tag = TEXT("NamedDiscreteExclusiveAction"),
-		const bool bVisualLoggerEnabled = false,
-		ULearningAgentsManagerListener* VisualLoggerListener = nullptr,
-		const int32 VisualLoggerAgentId = -1,
-		const FVector VisualLoggerLocation = FVector::ZeroVector,
-		const FLinearColor VisualLoggerColor = FLinearColor::Blue);
+		const FLearningAgentsVisualLoggerSettings& VisualLoggerSettings = FLearningAgentsVisualLoggerSettings());
 
-	UFUNCTION(BlueprintPure, Category = "LearningAgents", meta = (AdvancedDisplay = 2, DefaultToSelf = "VisualLoggerListener"))
+	UFUNCTION(BlueprintPure, Category = "LearningAgents", meta = (AdvancedDisplay = 2, AutoCreateRefTerm = "VisualLoggerSettings"))
 	static UE_API FLearningAgentsActionObjectElement MakeInclusiveDiscreteAction(
 		ULearningAgentsActionObject* Object,
 		const TArray<int32>& Indices,
 		const bool bActive = true, 
 		const FName Tag = TEXT("DiscreteInclusiveAction"),
-		const bool bVisualLoggerEnabled = false,
-		ULearningAgentsManagerListener* VisualLoggerListener = nullptr,
-		const int32 VisualLoggerAgentId = -1,
-		const FVector VisualLoggerLocation = FVector::ZeroVector,
-		const FLinearColor VisualLoggerColor = FLinearColor::Blue);
+		const FLearningAgentsVisualLoggerSettings& VisualLoggerSettings = FLearningAgentsVisualLoggerSettings());
 
 	static UE_API FLearningAgentsActionObjectElement MakeInclusiveDiscreteActionFromArrayView(
 		ULearningAgentsActionObject* Object,
 		const TArrayView<const int32> Indices,
 		const bool bActive = true,
 		const FName Tag = TEXT("DiscreteInclusiveAction"),
-		const bool bVisualLoggerEnabled = false,
-		ULearningAgentsManagerListener* VisualLoggerListener = nullptr,
-		const int32 VisualLoggerAgentId = -1,
-		const FVector VisualLoggerLocation = FVector::ZeroVector,
-		const FLinearColor VisualLoggerColor = FLinearColor::Blue);
+		const FLearningAgentsVisualLoggerSettings& VisualLoggerSettings = FLearningAgentsVisualLoggerSettings());
 
-	UFUNCTION(BlueprintPure, Category = "LearningAgents", meta = (AdvancedDisplay = 2, DefaultToSelf = "VisualLoggerListener"))
+	UFUNCTION(BlueprintPure, Category = "LearningAgents", meta = (AdvancedDisplay = 2, AutoCreateRefTerm = "VisualLoggerSettings"))
 	static UE_API FLearningAgentsActionObjectElement MakeNamedInclusiveDiscreteAction(
 		ULearningAgentsActionObject* Object,
 		const TArray<FName>& Names,
 		const bool bActive = true, 
 		const FName Tag = TEXT("NamedDiscreteInclusiveAction"),
-		const bool bVisualLoggerEnabled = false,
-		ULearningAgentsManagerListener* VisualLoggerListener = nullptr,
-		const int32 VisualLoggerAgentId = -1,
-		const FVector VisualLoggerLocation = FVector::ZeroVector,
-		const FLinearColor VisualLoggerColor = FLinearColor::Blue);
+		const FLearningAgentsVisualLoggerSettings& VisualLoggerSettings = FLearningAgentsVisualLoggerSettings());
 
 	static UE_API FLearningAgentsActionObjectElement MakeNamedInclusiveDiscreteActionFromArrayView(
 		ULearningAgentsActionObject* Object,
 		const TArrayView<const FName> Names,
 		const bool bActive = true, 
 		const FName Tag = TEXT("NamedDiscreteInclusiveAction"),
-		const bool bVisualLoggerEnabled = false,
-		ULearningAgentsManagerListener* VisualLoggerListener = nullptr,
-		const int32 VisualLoggerAgentId = -1,
-		const FVector VisualLoggerLocation = FVector::ZeroVector,
-		const FLinearColor VisualLoggerColor = FLinearColor::Blue);
+		const FLearningAgentsVisualLoggerSettings& VisualLoggerSettings = FLearningAgentsVisualLoggerSettings());
 
 	UFUNCTION(BlueprintPure, Category = "LearningAgents", meta = (AdvancedDisplay = 2))
 	static UE_API FLearningAgentsActionObjectElement MakeStructAction(ULearningAgentsActionObject* Object, const TMap<FName, FLearningAgentsActionObjectElement>& Elements, const bool bActive = true, const FName Tag = TEXT("StructAction"));
@@ -837,31 +899,23 @@ public:
 	UFUNCTION(BlueprintPure, Category = "LearningAgents", meta = (AdvancedDisplay = 3))
 	static UE_API FLearningAgentsActionObjectElement MakePairAction(ULearningAgentsActionObject* Object, const FLearningAgentsActionObjectElement Key, const FLearningAgentsActionObjectElement Value, const bool bActive = true, const FName Tag = TEXT("PairAction"));
 
-	UFUNCTION(BlueprintPure, Category = "LearningAgents", meta = (AdvancedDisplay = 2, DefaultToSelf = "VisualLoggerListener"))
+	UFUNCTION(BlueprintPure, Category = "LearningAgents", meta = (AdvancedDisplay = 2, AutoCreateRefTerm = "VisualLoggerSettings"))
 	static UE_API FLearningAgentsActionObjectElement MakeEnumAction(
 		ULearningAgentsActionObject* Object,
 		const UEnum* Enum,
 		const uint8 EnumValue,
 		const bool bActive = true, 
 		const FName Tag = TEXT("EnumAction"),
-		const bool bVisualLoggerEnabled = false,
-		ULearningAgentsManagerListener* VisualLoggerListener = nullptr,
-		const int32 VisualLoggerAgentId = -1,
-		const FVector VisualLoggerLocation = FVector::ZeroVector,
-		const FLinearColor VisualLoggerColor = FLinearColor::Blue);
+		const FLearningAgentsVisualLoggerSettings& VisualLoggerSettings = FLearningAgentsVisualLoggerSettings());
 
-	UFUNCTION(BlueprintPure, Category = "LearningAgents", meta = (AdvancedDisplay = 2, DefaultToSelf = "VisualLoggerListener"))
+	UFUNCTION(BlueprintPure, Category = "LearningAgents", meta = (AdvancedDisplay = 2, AutoCreateRefTerm = "VisualLoggerSettings"))
 	static UE_API FLearningAgentsActionObjectElement MakeBitmaskAction(
 		ULearningAgentsActionObject* Object,
 		const UEnum* Enum,
 		const int32 BitmaskValue,
 		const bool bActive = true, 
 		const FName Tag = TEXT("BitmaskAction"),
-		const bool bVisualLoggerEnabled = false,
-		ULearningAgentsManagerListener* VisualLoggerListener = nullptr,
-		const int32 VisualLoggerAgentId = -1,
-		const FVector VisualLoggerLocation = FVector::ZeroVector,
-		const FLinearColor VisualLoggerColor = FLinearColor::Blue);
+		const FLearningAgentsVisualLoggerSettings& VisualLoggerSettings = FLearningAgentsVisualLoggerSettings());
 
 	UFUNCTION(BlueprintPure, Category = "LearningAgents", meta = (AdvancedDisplay = 3))
 	static UE_API FLearningAgentsActionObjectElement MakeOptionalAction(ULearningAgentsActionObject* Object, const FLearningAgentsActionObjectElement Element, const ELearningAgentsOptionalAction Option, const bool bActive = true, const FName Tag = TEXT("OptionalAction"));
@@ -884,153 +938,102 @@ public:
 	UFUNCTION(BlueprintPure, Category = "LearningAgents", meta = (AdvancedDisplay = 2))
 	static UE_API FLearningAgentsActionObjectElement MakeEncodingAction(ULearningAgentsActionObject* Object, const FLearningAgentsActionObjectElement Element, const bool bActive = true, const FName Tag = TEXT("EncodingAction"));
 
-	UFUNCTION(BlueprintPure, Category = "LearningAgents", meta = (AdvancedDisplay = 2, DefaultToSelf = "VisualLoggerListener"))
+	UFUNCTION(BlueprintPure, Category = "LearningAgents", meta = (AdvancedDisplay = 2, AutoCreateRefTerm = "VisualLoggerSettings"))
 	static UE_API FLearningAgentsActionObjectElement MakeBoolAction(
 		ULearningAgentsActionObject* Object,
 		const bool bValue,
 		const bool bActive = true, 
 		const FName Tag = TEXT("BoolAction"),
-		const bool bVisualLoggerEnabled = false,
-		ULearningAgentsManagerListener * VisualLoggerListener = nullptr,
-		const int32 VisualLoggerAgentId = -1,
-		const FVector VisualLoggerLocation = FVector::ZeroVector,
-		const FLinearColor VisualLoggerColor = FLinearColor::Blue);
+		const FLearningAgentsVisualLoggerSettings& VisualLoggerSettings = FLearningAgentsVisualLoggerSettings());
 
-	UFUNCTION(BlueprintPure, Category = "LearningAgents", meta = (AdvancedDisplay = 2, DefaultToSelf = "VisualLoggerListener"))
+	UFUNCTION(BlueprintPure, Category = "LearningAgents", meta = (AdvancedDisplay = 2, AutoCreateRefTerm = "VisualLoggerSettings"))
 	static UE_API FLearningAgentsActionObjectElement MakeFloatAction(
 		ULearningAgentsActionObject* Object,
 		const float Value,
-		const bool bActive = true, 
+		const bool bActive = true,
 		const FName Tag = TEXT("FloatAction"),
-		const bool bVisualLoggerEnabled = false,
-		ULearningAgentsManagerListener* VisualLoggerListener = nullptr,
-		const int32 VisualLoggerAgentId = -1,
-		const FVector VisualLoggerLocation = FVector::ZeroVector,
-		const FLinearColor VisualLoggerColor = FLinearColor::Blue);
+		const FLearningAgentsVisualLoggerSettings& VisualLoggerSettings = FLearningAgentsVisualLoggerSettings());
 
-	UFUNCTION(BlueprintPure, Category = "LearningAgents", meta = (AdvancedDisplay = 3, DefaultToSelf = "VisualLoggerListener"))
+	UFUNCTION(BlueprintPure, Category = "LearningAgents", meta = (AdvancedDisplay = 3, AutoCreateRefTerm = "VisualLoggerSettings"))
 	static UE_API FLearningAgentsActionObjectElement MakeLocationAction(
 		ULearningAgentsActionObject* Object,
 		const FVector Location,
 		const FTransform RelativeTransform = FTransform(),
-		const bool bActive = true, 
+		const bool bActive = true,
 		const FName Tag = TEXT("LocationAction"),
-		const bool bVisualLoggerEnabled = false,
-		ULearningAgentsManagerListener* VisualLoggerListener = nullptr,
-		const int32 VisualLoggerAgentId = -1,
-		const FVector VisualLoggerLocation = FVector::ZeroVector,
-		const FLinearColor VisualLoggerColor = FLinearColor::Blue);
+		const FLearningAgentsVisualLoggerSettings& VisualLoggerSettings = FLearningAgentsVisualLoggerSettings());
 
-	UFUNCTION(BlueprintPure, Category = "LearningAgents", meta = (AdvancedDisplay = 3, DefaultToSelf = "VisualLoggerListener"))
+	UFUNCTION(BlueprintPure, Category = "LearningAgents", meta = (AdvancedDisplay = 3, AutoCreateRefTerm = "VisualLoggerSettings"))
 	static UE_API FLearningAgentsActionObjectElement MakeRotationAction(
 		ULearningAgentsActionObject* Object,
 		const FRotator Rotation,
 		const FRotator RelativeRotation = FRotator::ZeroRotator,
-		const bool bActive = true, 
+		const bool bActive = true,
 		const FName Tag = TEXT("RotationAction"),
-		const bool bVisualLoggerEnabled = false,
-		ULearningAgentsManagerListener* VisualLoggerListener = nullptr,
-		const int32 VisualLoggerAgentId = -1,
-		const FVector VisualLoggerRotationLocation = FVector::ZeroVector,
-		const FVector VisualLoggerLocation = FVector::ZeroVector,
-		const FLinearColor VisualLoggerColor = FLinearColor::Blue);
-
-	UFUNCTION(BlueprintPure, Category = "LearningAgents", meta = (AdvancedDisplay = 3, DefaultToSelf = "VisualLoggerListener"))
+		const FLearningAgentsVisualLoggerSettings& VisualLoggerSettings = FLearningAgentsVisualLoggerSettings());
+
+	UFUNCTION(BlueprintPure, Category = "LearningAgents", meta = (AdvancedDisplay = 3, AutoCreateRefTerm = "VisualLoggerSettings"))
 	static UE_API FLearningAgentsActionObjectElement MakeRotationActionFromQuat(
 		ULearningAgentsActionObject* Object,
 		const FQuat Rotation,
 		const FQuat RelativeRotation,
-		const bool bActive = true, 
+		const bool bActive = true,
 		const FName Tag = TEXT("RotationAction"),
-		const bool bVisualLoggerEnabled = false,
-		ULearningAgentsManagerListener* VisualLoggerListener = nullptr,
-		const int32 VisualLoggerAgentId = -1,
-		const FVector VisualLoggerRotationLocation = FVector::ZeroVector,
-		const FVector VisualLoggerLocation = FVector::ZeroVector,
-		const FLinearColor VisualLoggerColor = FLinearColor::Blue);
-
-	UFUNCTION(BlueprintPure, Category = "LearningAgents", meta = (AdvancedDisplay = 3, DefaultToSelf = "VisualLoggerListener"))
+		const FLearningAgentsVisualLoggerSettings& VisualLoggerSettings = FLearningAgentsVisualLoggerSettings());
+
+	UFUNCTION(BlueprintPure, Category = "LearningAgents", meta = (AdvancedDisplay = 3, AutoCreateRefTerm = "VisualLoggerSettings"))
 	static UE_API FLearningAgentsActionObjectElement MakeScaleAction(
 		ULearningAgentsActionObject* Object,
 		const FVector Scale,
 		const FVector RelativeScale = FVector(1, 1, 1),
-		const bool bActive = true, 
+		const bool bActive = true,
 		const FName Tag = TEXT("ScaleAction"),
-		const bool bVisualLoggerEnabled = false,
-		ULearningAgentsManagerListener* VisualLoggerListener = nullptr,
-		const int32 VisualLoggerAgentId = -1,
-		const FVector VisualLoggerLocation = FVector::ZeroVector,
-		const FLinearColor VisualLoggerColor = FLinearColor::Blue);
+		const FLearningAgentsVisualLoggerSettings& VisualLoggerSettings = FLearningAgentsVisualLoggerSettings());
 
-	UFUNCTION(BlueprintPure, Category = "LearningAgents", meta = (AdvancedDisplay = 3, DefaultToSelf = "VisualLoggerListener"))
+	UFUNCTION(BlueprintPure, Category = "LearningAgents", meta = (AdvancedDisplay = 3, AutoCreateRefTerm = "VisualLoggerSettings"))
 	static UE_API FLearningAgentsActionObjectElement MakeTransformAction(
 		ULearningAgentsActionObject* Object,
 		const FTransform Transform,
 		const FTransform RelativeTransform = FTransform(),
 		const bool bActive = true, 
 		const FName Tag = TEXT("TransformAction"),
-		const bool bVisualLoggerEnabled = false,
-		ULearningAgentsManagerListener* VisualLoggerListener = nullptr,
-		const int32 VisualLoggerAgentId = -1,
-		const FVector VisualLoggerLocation = FVector::ZeroVector,
-		const FLinearColor VisualLoggerColor = FLinearColor::Blue);
+		const FLearningAgentsVisualLoggerSettings& VisualLoggerSettings = FLearningAgentsVisualLoggerSettings());
 
-	UFUNCTION(BlueprintPure, Category = "LearningAgents", meta = (AdvancedDisplay = 3, DefaultToSelf = "VisualLoggerListener"))
+	UFUNCTION(BlueprintPure, Category = "LearningAgents", meta = (AdvancedDisplay = 3, AutoCreateRefTerm = "VisualLoggerSettings"))
 	static UE_API FLearningAgentsActionObjectElement MakeAngleAction(
 		ULearningAgentsActionObject* Object,
 		const float Angle,
 		const float RelativeAngle = 0.0f,
-		const bool bActive = true, 
+		const bool bActive = true,
 		const FName Tag = TEXT("AngleAction"),
-		const bool bVisualLoggerEnabled = false,
-		ULearningAgentsManagerListener* VisualLoggerListener = nullptr,
-		const int32 VisualLoggerAgentId = -1,
-		const FVector VisualLoggerAngleLocation = FVector::ZeroVector,
-		const FVector VisualLoggerLocation = FVector::ZeroVector,
-		const FLinearColor VisualLoggerColor = FLinearColor::Blue);
-
-	UFUNCTION(BlueprintPure, Category = "LearningAgents", meta = (AdvancedDisplay = 3, DefaultToSelf = "VisualLoggerListener"))
+		const FLearningAgentsVisualLoggerSettings& VisualLoggerSettings = FLearningAgentsVisualLoggerSettings());
+
+	UFUNCTION(BlueprintPure, Category = "LearningAgents", meta = (AdvancedDisplay = 3, AutoCreateRefTerm = "VisualLoggerSettings"))
 	static UE_API FLearningAgentsActionObjectElement MakeAngleActionRadians(
 		ULearningAgentsActionObject* Object,
 		const float Angle,
 		const float RelativeAngle = 0.0f,
-		const bool bActive = true, 
+		const bool bActive = true,
 		const FName Tag = TEXT("AngleAction"),
-		const bool bVisualLoggerEnabled = false,
-		ULearningAgentsManagerListener* VisualLoggerListener = nullptr,
-		const int32 VisualLoggerAgentId = -1,
-		const FVector VisualLoggerAngleLocation = FVector::ZeroVector,
-		const FVector VisualLoggerLocation = FVector::ZeroVector,
-		const FLinearColor VisualLoggerColor = FLinearColor::Blue);
-
-	UFUNCTION(BlueprintPure, Category = "LearningAgents", meta = (AdvancedDisplay = 3, DefaultToSelf = "VisualLoggerListener"))
+		const FLearningAgentsVisualLoggerSettings& VisualLoggerSettings = FLearningAgentsVisualLoggerSettings());
+
+	UFUNCTION(BlueprintPure, Category = "LearningAgents", meta = (AdvancedDisplay = 3, AutoCreateRefTerm = "VisualLoggerSettings"))
 	static UE_API FLearningAgentsActionObjectElement MakeVelocityAction(
 		ULearningAgentsActionObject* Object,
 		const FVector Velocity,
 		const FTransform RelativeTransform = FTransform(),
-		const bool bActive = true, 
+		const bool bActive = true,
 		const FName Tag = TEXT("VelocityAction"),
-		const bool bVisualLoggerEnabled = false,
-		ULearningAgentsManagerListener* VisualLoggerListener = nullptr,
-		const int32 VisualLoggerAgentId = -1,
-		const FVector VisualLoggerVelocityLocation = FVector::ZeroVector,
-		const FVector VisualLoggerLocation = FVector::ZeroVector,
-		const FLinearColor VisualLoggerColor = FLinearColor::Blue);
-
-	UFUNCTION(BlueprintPure, Category = "LearningAgents", meta = (AdvancedDisplay = 3, DefaultToSelf = "VisualLoggerListener"))
+		const FLearningAgentsVisualLoggerSettings& VisualLoggerSettings = FLearningAgentsVisualLoggerSettings());
+
+	UFUNCTION(BlueprintPure, Category = "LearningAgents", meta = (AdvancedDisplay = 3, AutoCreateRefTerm = "VisualLoggerSettings"))
 	static UE_API FLearningAgentsActionObjectElement MakeDirectionAction(
 		ULearningAgentsActionObject* Object,
 		const FVector Direction,
 		const FTransform RelativeTransform = FTransform(),
-		const bool bActive = true, 
+		const bool bActive = true,
 		const FName Tag = TEXT("DirectionAction"),
-		const bool bVisualLoggerEnabled = false,
-		ULearningAgentsManagerListener* VisualLoggerListener = nullptr,
-		const int32 VisualLoggerAgentId = -1,
-		const FVector VisualLoggerDirectionLocation = FVector::ZeroVector,
-		const FVector VisualLoggerLocation = FVector::ZeroVector,
-		const float VisualLoggerArrowLength = 100.0f,
-		const FLinearColor VisualLoggerColor = FLinearColor::Blue);
+		const FLearningAgentsVisualLoggerSettings& VisualLoggerSettings = FLearningAgentsVisualLoggerSettings());
 
 public:
 
@@ -1166,19 +1169,15 @@ public:
 	 * @param VisualLoggerColor The color for the visual logger display.
 	 * @return true if the provided Element is the correct type, otherwise false.
 	 */
-	UFUNCTION(BlueprintPure = false, Category = "LearningAgents", meta = (AdvancedDisplay = 3, ReturnDisplayName = "Success", DefaultToSelf = "VisualLoggerListener"))
+	UFUNCTION(BlueprintPure = false, Category = "LearningAgents", meta = (AdvancedDisplay = 3, ReturnDisplayName = "Success", AutoCreateRefTerm = "VisualLoggerSettings"))
 	static UE_API bool GetContinuousAction(
-		TArray<float>& OutValues, 
-		const ULearningAgentsActionObject* Object, 
-		const FLearningAgentsActionObjectElement Element, 
-		const bool bActive = true, 
+		TArray<float>& OutValues,
+		const ULearningAgentsActionObject* Object,
+		const FLearningAgentsActionObjectElement Element,
+		const bool bActive = true,
 		const FName Tag = TEXT("ContinuousAction"),
-		const bool bVisualLoggerEnabled = false,
-		ULearningAgentsManagerListener* VisualLoggerListener = nullptr,
-		const int32 VisualLoggerAgentId = -1,
-		const FVector VisualLoggerLocation = FVector::ZeroVector,
-		const FLinearColor VisualLoggerColor = FLinearColor::Blue);
-	
+		const FLearningAgentsVisualLoggerSettings& VisualLoggerSettings = FLearningAgentsVisualLoggerSettings());
+
 	/**
 	 * Get the values of a continuous action. The OutValues ArrayView must be the correct size.
 	 *
@@ -1195,16 +1194,12 @@ public:
 	 * @return true if the provided Element is the correct type, otherwise false.
 	 */
 	static UE_API bool GetContinuousActionToArrayView(
-		TArrayView<float> OutValues, 
-		const ULearningAgentsActionObject* Object, 
-		const FLearningAgentsActionObjectElement Element, 
-		const bool bActive = true, 
+		TArrayView<float> OutValues,
+		const ULearningAgentsActionObject* Object,
+		const FLearningAgentsActionObjectElement Element,
+		const bool bActive = true,
 		const FName Tag = TEXT("ContinuousAction"),
-		const bool bVisualLoggerEnabled = false,
-		ULearningAgentsManagerListener* VisualLoggerListener = nullptr,
-		const int32 VisualLoggerAgentId = -1,
-		const FVector VisualLoggerLocation = FVector::ZeroVector,
-		const FLinearColor VisualLoggerColor = FLinearColor::Blue);
+		const FLearningAgentsVisualLoggerSettings& VisualLoggerSettings = FLearningAgentsVisualLoggerSettings());
 
 	/**
 	 * Get the index for an exclusive discrete action.
@@ -1221,18 +1216,14 @@ public:
 	 * @param VisualLoggerColor The color for the visual logger display.
 	 * @return true if the provided Element is the correct type, otherwise false.
 	 */
-	UFUNCTION(BlueprintPure = false, Category = "LearningAgents", meta = (AdvancedDisplay = 3, ReturnDisplayName = "Success", DefaultToSelf = "VisualLoggerListener"))
+	UFUNCTION(BlueprintPure = false, Category = "LearningAgents", meta = (AdvancedDisplay = 3, ReturnDisplayName = "Success", AutoCreateRefTerm = "VisualLoggerSettings"))
 	static UE_API bool GetExclusiveDiscreteAction(
 		int32& OutIndex, 
 		const ULearningAgentsActionObject* Object, 
 		const FLearningAgentsActionObjectElement Element, 
 		const bool bActive = true, 
 		const FName Tag = TEXT("DiscreteExclusiveAction"),
-		const bool bVisualLoggerEnabled = false,
-		ULearningAgentsManagerListener* VisualLoggerListener = nullptr,
-		const int32 VisualLoggerAgentId = -1,
-		const FVector VisualLoggerLocation = FVector::ZeroVector,
-		const FLinearColor VisualLoggerColor = FLinearColor::Blue);
+		const FLearningAgentsVisualLoggerSettings& VisualLoggerSettings = FLearningAgentsVisualLoggerSettings());
 
 	/**
 	 * Get the name for a named exclusive discrete action.
@@ -1249,18 +1240,14 @@ public:
 	 * @param VisualLoggerColor The color for the visual logger display.
 	 * @return true if the provided Element is the correct type, otherwise false.
 	 */
-	UFUNCTION(BlueprintPure = false, Category = "LearningAgents", meta = (AdvancedDisplay = 3, ReturnDisplayName = "Success", DefaultToSelf = "VisualLoggerListener"))
+	UFUNCTION(BlueprintPure = false, Category = "LearningAgents", meta = (AdvancedDisplay = 3, ReturnDisplayName = "Success", AutoCreateRefTerm = "VisualLoggerSettings"))
 	static UE_API bool GetNamedExclusiveDiscreteAction(
 		FName& OutName,
 		const ULearningAgentsActionObject* Object,
 		const FLearningAgentsActionObjectElement Element,
 		const bool bActive = true, 
 		const FName Tag = TEXT("NamedDiscreteExclusiveAction"),
-		const bool bVisualLoggerEnabled = false,
-		ULearningAgentsManagerListener* VisualLoggerListener = nullptr,
-		const int32 VisualLoggerAgentId = -1,
-		const FVector VisualLoggerLocation = FVector::ZeroVector,
-		const FLinearColor VisualLoggerColor = FLinearColor::Blue);
+		const FLearningAgentsVisualLoggerSettings& VisualLoggerSettings = FLearningAgentsVisualLoggerSettings());
 
 	/**
 	 * Get the number of indices for an inclusive discrete action.
@@ -1290,18 +1277,14 @@ public:
 	 * @param VisualLoggerColor The color for the visual logger display.
 	 * @return true if the provided Element is the correct type, otherwise false.
 	 */
-	UFUNCTION(BlueprintPure = false, Category = "LearningAgents", meta = (AdvancedDisplay = 3, ReturnDisplayName = "Success", DefaultToSelf = "VisualLoggerListener"))
+	UFUNCTION(BlueprintPure = false, Category = "LearningAgents", meta = (AdvancedDisplay = 3, ReturnDisplayName = "Success", AutoCreateRefTerm = "VisualLoggerSettings"))
 	static UE_API bool GetInclusiveDiscreteAction(
 		TArray<int32>& OutIndices, 
 		const ULearningAgentsActionObject* Object, 
 		const FLearningAgentsActionObjectElement Element, 
 		const bool bActive = true, 
 		const FName Tag = TEXT("DiscreteInclusiveAction"),
-		const bool bVisualLoggerEnabled = false,
-		ULearningAgentsManagerListener* VisualLoggerListener = nullptr,
-		const int32 VisualLoggerAgentId = -1,
-		const FVector VisualLoggerLocation = FVector::ZeroVector,
-		const FLinearColor VisualLoggerColor = FLinearColor::Blue);
+		const FLearningAgentsVisualLoggerSettings& VisualLoggerSettings = FLearningAgentsVisualLoggerSettings());
 	
 	/**
 	 * Get the indices for an inclusive discrete action. The OutIndices ArrayView must be the correct size.
@@ -1324,11 +1307,7 @@ public:
 		const FLearningAgentsActionObjectElement Element, 
 		const bool bActive = true, 
 		const FName Tag = TEXT("DiscreteInclusiveAction"),
-		const bool bVisualLoggerEnabled = false,
-		ULearningAgentsManagerListener* VisualLoggerListener = nullptr,
-		const int32 VisualLoggerAgentId = -1,
-		const FVector VisualLoggerLocation = FVector::ZeroVector,
-		const FLinearColor VisualLoggerColor = FLinearColor::Blue);
+		const FLearningAgentsVisualLoggerSettings& VisualLoggerSettings = FLearningAgentsVisualLoggerSettings());
 
 	/**
 	 * Get the number of names for a named inclusive discrete action.
@@ -1358,18 +1337,14 @@ public:
 	 * @param VisualLoggerColor The color for the visual logger display.
 	 * @return true if the provided Element is the correct type, otherwise false.
 	 */
-	UFUNCTION(BlueprintPure = false, Category = "LearningAgents", meta = (AdvancedDisplay = 3, ReturnDisplayName = "Success", DefaultToSelf = "VisualLoggerListener"))
+	UFUNCTION(BlueprintPure = false, Category = "LearningAgents", meta = (AdvancedDisplay = 3, ReturnDisplayName = "Success", AutoCreateRefTerm = "VisualLoggerSettings"))
 	static UE_API bool GetNamedInclusiveDiscreteAction(
 		TArray<FName>& OutNames,
 		const ULearningAgentsActionObject* Object,
 		const FLearningAgentsActionObjectElement Element,
 		const bool bActive = true, 
 		const FName Tag = TEXT("NamedDiscreteInclusiveAction"),
-		const bool bVisualLoggerEnabled = false,
-		ULearningAgentsManagerListener* VisualLoggerListener = nullptr,
-		const int32 VisualLoggerAgentId = -1,
-		const FVector VisualLoggerLocation = FVector::ZeroVector,
-		const FLinearColor VisualLoggerColor = FLinearColor::Blue);
+		const FLearningAgentsVisualLoggerSettings& VisualLoggerSettings = FLearningAgentsVisualLoggerSettings());
 
 	/**
 	 * Get the names for a named inclusive discrete action. The OutNames ArrayView must be the correct size.
@@ -1392,11 +1367,7 @@ public:
 		const FLearningAgentsActionObjectElement Element,
 		const bool bActive = true, 
 		const FName Tag = TEXT("NamedDiscreteInclusiveAction"),
-		const bool bVisualLoggerEnabled = false,
-		ULearningAgentsManagerListener* VisualLoggerListener = nullptr,
-		const int32 VisualLoggerAgentId = -1,
-		const FVector VisualLoggerLocation = FVector::ZeroVector,
-		const FLinearColor VisualLoggerColor = FLinearColor::Blue);
+		const FLearningAgentsVisualLoggerSettings& VisualLoggerSettings = FLearningAgentsVisualLoggerSettings());
 
 	/**
 	 * Get the number of sub-actions for a struct action.
@@ -1599,7 +1570,7 @@ public:
 	 * @param VisualLoggerColor The color for the visual logger display.
 	 * @return true if the provided Element is the correct type, otherwise false.
 	 */
-	UFUNCTION(BlueprintPure = false, Category = "LearningAgents", meta = (AdvancedDisplay = 4, ReturnDisplayName = "Success", DefaultToSelf = "VisualLoggerListener"))
+	UFUNCTION(BlueprintPure = false, Category = "LearningAgents", meta = (AdvancedDisplay = 4, ReturnDisplayName = "Success", AutoCreateRefTerm = "VisualLoggerSettings"))
 	static UE_API bool GetEnumAction(
 		uint8& OutEnumValue, 
 		const ULearningAgentsActionObject* Object, 
@@ -1607,11 +1578,7 @@ public:
 		const UEnum* Enum, 
 		const bool bActive = true, 
 		const FName Tag = TEXT("EnumAction"),
-		const bool bVisualLoggerEnabled = false,
-		ULearningAgentsManagerListener* VisualLoggerListener = nullptr,
-		const int32 VisualLoggerAgentId = -1,
-		const FVector VisualLoggerLocation = FVector::ZeroVector,
-		const FLinearColor VisualLoggerColor = FLinearColor::Blue);
+		const FLearningAgentsVisualLoggerSettings& VisualLoggerSettings = FLearningAgentsVisualLoggerSettings());
 
 	/**
 	 * Get the bitmask value of a bitmask action.
@@ -1629,7 +1596,7 @@ public:
 	 * @param VisualLoggerColor The color for the visual logger display.
 	 * @return true if the provided Element is the correct type, otherwise false.
 	 */
-	UFUNCTION(BlueprintPure = false, Category = "LearningAgents", meta = (AdvancedDisplay = 4, ReturnDisplayName = "Success", DefaultToSelf = "VisualLoggerListener"))
+	UFUNCTION(BlueprintPure = false, Category = "LearningAgents", meta = (AdvancedDisplay = 4, ReturnDisplayName = "Success", AutoCreateRefTerm = "VisualLoggerSettings"))
 	static UE_API bool GetBitmaskAction(
 		int32& OutBitmaskValue, 
 		const ULearningAgentsActionObject* Object, 
@@ -1637,11 +1604,7 @@ public:
 		const UEnum* Enum, 
 		const bool bActive = true, 
 		const FName Tag = TEXT("BitmaskAction"),
-		const bool bVisualLoggerEnabled = false,
-		ULearningAgentsManagerListener* VisualLoggerListener = nullptr,
-		const int32 VisualLoggerAgentId = -1,
-		const FVector VisualLoggerLocation = FVector::ZeroVector,
-		const FLinearColor VisualLoggerColor = FLinearColor::Blue);
+		const FLearningAgentsVisualLoggerSettings& VisualLoggerSettings = FLearningAgentsVisualLoggerSettings());
 
 	/**
 	 * Get the sub-action of an option action.
@@ -1699,18 +1662,14 @@ public:
 	 * @param VisualLoggerColor The color for the visual logger display.
 	 * @return true if the provided Element is the correct type, otherwise false.
 	 */
-	UFUNCTION(BlueprintPure = false, Category = "LearningAgents", meta = (AdvancedDisplay = 3, ReturnDisplayName = "Success", DefaultToSelf = "VisualLoggerListener"))
+	UFUNCTION(BlueprintPure = false, Category = "LearningAgents", meta = (AdvancedDisplay = 3, ReturnDisplayName = "Success", AutoCreateRefTerm = "VisualLoggerSettings"))
 	static UE_API bool GetBoolAction(
 		bool& bOutValue, 
 		const ULearningAgentsActionObject* Object, 
 		const FLearningAgentsActionObjectElement Element, 
 		const bool bActive = true, 
 		const FName Tag = TEXT("BoolAction"),
-		const bool bVisualLoggerEnabled = false,
-		ULearningAgentsManagerListener* VisualLoggerListener = nullptr,
-		const int32 VisualLoggerAgentId = -1,
-		const FVector VisualLoggerLocation = FVector::ZeroVector,
-		const FLinearColor VisualLoggerColor = FLinearColor::Blue);
+		const FLearningAgentsVisualLoggerSettings& VisualLoggerSettings = FLearningAgentsVisualLoggerSettings());
 
 	/**
 	 * Get the value for a float action.
@@ -1727,18 +1686,14 @@ public:
 	 * @param VisualLoggerColor The color for the visual logger display.
 	 * @return true if the provided Element is the correct type, otherwise false.
 	 */
-	UFUNCTION(BlueprintPure = false, Category = "LearningAgents", meta = (AdvancedDisplay = 3, ReturnDisplayName = "Success", DefaultToSelf = "VisualLoggerListener"))
+	UFUNCTION(BlueprintPure = false, Category = "LearningAgents", meta = (AdvancedDisplay = 3, ReturnDisplayName = "Success", AutoCreateRefTerm = "VisualLoggerSettings"))
 	static UE_API bool GetFloatAction(
 		float& OutValue, 
 		const ULearningAgentsActionObject* Object, 
 		const FLearningAgentsActionObjectElement Element, 
 		const bool bActive = true, 
 		const FName Tag = TEXT("FloatAction"),
-		const bool bVisualLoggerEnabled = false,
-		ULearningAgentsManagerListener* VisualLoggerListener = nullptr,
-		const int32 VisualLoggerAgentId = -1,
-		const FVector VisualLoggerLocation = FVector::ZeroVector,
-		const FLinearColor VisualLoggerColor = FLinearColor::Blue);
+		const FLearningAgentsVisualLoggerSettings& VisualLoggerSettings = FLearningAgentsVisualLoggerSettings());
 
 	/**
 	 * Get the value for a location action.
@@ -1756,7 +1711,7 @@ public:
 	 * @param VisualLoggerColor The color for the visual logger display.
 	 * @return true if the provided Element is the correct type, otherwise false.
 	 */
-	UFUNCTION(BlueprintPure = false, Category = "LearningAgents", meta = (AdvancedDisplay = 4, ReturnDisplayName = "Success", DefaultToSelf = "VisualLoggerListener"))
+	UFUNCTION(BlueprintPure = false, Category = "LearningAgents", meta = (AdvancedDisplay = 4, ReturnDisplayName = "Success", AutoCreateRefTerm = "VisualLoggerSettings"))
 	static UE_API bool GetLocationAction(
 		FVector& OutLocation, 
 		const ULearningAgentsActionObject* Object, 
@@ -1764,11 +1719,7 @@ public:
 		const FTransform RelativeTransform = FTransform(), 
 		const bool bActive = true, 
 		const FName Tag = TEXT("LocationAction"),
-		const bool bVisualLoggerEnabled = false,
-		ULearningAgentsManagerListener* VisualLoggerListener = nullptr,
-		const int32 VisualLoggerAgentId = -1,
-		const FVector VisualLoggerLocation = FVector::ZeroVector,
-		const FLinearColor VisualLoggerColor = FLinearColor::Blue);
+		const FLearningAgentsVisualLoggerSettings& VisualLoggerSettings = FLearningAgentsVisualLoggerSettings());
 
 	/**
 	 * Get the value for a rotation action.
@@ -1787,7 +1738,7 @@ public:
 	 * @param VisualLoggerColor The color for the visual logger display.
 	 * @return true if the provided Element is the correct type, otherwise false.
 	 */
-	UFUNCTION(BlueprintPure = false, Category = "LearningAgents", meta = (AdvancedDisplay = 4, ReturnDisplayName = "Success", DefaultToSelf = "VisualLoggerListener"))
+	UFUNCTION(BlueprintPure = false, Category = "LearningAgents", meta = (AdvancedDisplay = 4, ReturnDisplayName = "Success", AutoCreateRefTerm = "VisualLoggerSettings"))
 	static UE_API bool GetRotationAction(
 		FRotator& OutRotation, 
 		const ULearningAgentsActionObject* Object, 
@@ -1795,12 +1746,7 @@ public:
 		const FRotator RelativeRotation = FRotator::ZeroRotator, 
 		const bool bActive = true, 
 		const FName Tag = TEXT("RotationAction"),
-		const bool bVisualLoggerEnabled = false,
-		ULearningAgentsManagerListener* VisualLoggerListener = nullptr,
-		const int32 VisualLoggerAgentId = -1,
-		const FVector VisualLoggerRotationLocation = FVector::ZeroVector,
-		const FVector VisualLoggerLocation = FVector::ZeroVector,
-		const FLinearColor VisualLoggerColor = FLinearColor::Blue);
+		const FLearningAgentsVisualLoggerSettings& VisualLoggerSettings = FLearningAgentsVisualLoggerSettings());
 
 	/**
 	 * Get the value for a rotation action as a quaternion.
@@ -1819,7 +1765,7 @@ public:
 	 * @param VisualLoggerColor The color for the visual logger display.
 	 * @return true if the provided Element is the correct type, otherwise false.
 	 */
-	UFUNCTION(BlueprintPure = false, Category = "LearningAgents", meta = (AdvancedDisplay = 4, ReturnDisplayName = "Success", DefaultToSelf = "VisualLoggerListener"))
+	UFUNCTION(BlueprintPure = false, Category = "LearningAgents", meta = (AdvancedDisplay = 4, ReturnDisplayName = "Success", AutoCreateRefTerm = "VisualLoggerSettings"))
 	static UE_API bool GetRotationActionAsQuat(
 		FQuat& OutRotation, 
 		const ULearningAgentsActionObject* Object, 
@@ -1827,12 +1773,7 @@ public:
 		const FQuat RelativeRotation, 
 		const bool bActive = true, 
 		const FName Tag = TEXT("RotationAction"),
-		const bool bVisualLoggerEnabled = false,
-		ULearningAgentsManagerListener* VisualLoggerListener = nullptr,
-		const int32 VisualLoggerAgentId = -1,
-		const FVector VisualLoggerRotationLocation = FVector::ZeroVector,
-		const FVector VisualLoggerLocation = FVector::ZeroVector,
-		const FLinearColor VisualLoggerColor = FLinearColor::Blue);
+		const FLearningAgentsVisualLoggerSettings& VisualLoggerSettings = FLearningAgentsVisualLoggerSettings());
 
 	/**
 	 * Get the value for a scale action.
@@ -1850,7 +1791,7 @@ public:
 	 * @param VisualLoggerColor The color for the visual logger display.
 	 * @return true if the provided Element is the correct type, otherwise false.
 	 */
-	UFUNCTION(BlueprintPure = false, Category = "LearningAgents", meta = (AdvancedDisplay = 4, ReturnDisplayName = "Success", DefaultToSelf = "VisualLoggerListener"))
+	UFUNCTION(BlueprintPure = false, Category = "LearningAgents", meta = (AdvancedDisplay = 4, ReturnDisplayName = "Success", AutoCreateRefTerm = "VisualLoggerSettings"))
 	static UE_API bool GetScaleAction(
 		FVector& OutScale, 
 		const ULearningAgentsActionObject* Object, 
@@ -1858,11 +1799,7 @@ public:
 		const FVector RelativeScale = FVector(1, 1, 1), 
 		const bool bActive = true, 
 		const FName Tag = TEXT("ScaleAction"),
-		const bool bVisualLoggerEnabled = false,
-		ULearningAgentsManagerListener* VisualLoggerListener = nullptr,
-		const int32 VisualLoggerAgentId = -1,
-		const FVector VisualLoggerLocation = FVector::ZeroVector,
-		const FLinearColor VisualLoggerColor = FLinearColor::Blue);
+		const FLearningAgentsVisualLoggerSettings& VisualLoggerSettings = FLearningAgentsVisualLoggerSettings());
 
 	/**
 	 * Get the value for a transform action.
@@ -1880,7 +1817,7 @@ public:
 	 * @param VisualLoggerColor The color for the visual logger display.
 	 * @return true if the provided Element is the correct type, otherwise false.
 	 */
-	UFUNCTION(BlueprintPure = false, Category = "LearningAgents", meta = (AdvancedDisplay = 6, ReturnDisplayName = "Success", DefaultToSelf = "VisualLoggerListener"))
+	UFUNCTION(BlueprintPure = false, Category = "LearningAgents", meta = (AdvancedDisplay = 6, ReturnDisplayName = "Success", AutoCreateRefTerm = "VisualLoggerSettings"))
 	static UE_API bool GetTransformAction(
 		FTransform& OutTransform, 
 		const ULearningAgentsActionObject* Object, 
@@ -1888,11 +1825,7 @@ public:
 		const FTransform RelativeTransform = FTransform(), 
 		const bool bActive = true, 
 		const FName Tag = TEXT("TransformAction"),
-		const bool bVisualLoggerEnabled = false,
-		ULearningAgentsManagerListener* VisualLoggerListener = nullptr,
-		const int32 VisualLoggerAgentId = -1,
-		const FVector VisualLoggerLocation = FVector::ZeroVector,
-		const FLinearColor VisualLoggerColor = FLinearColor::Blue);
+		const FLearningAgentsVisualLoggerSettings& VisualLoggerSettings = FLearningAgentsVisualLoggerSettings());
 
 	/**
 	 * Get the value for an angle action. Returned angle is in degrees.
@@ -1911,7 +1844,7 @@ public:
 	 * @param VisualLoggerColor The color for the visual logger display.
 	 * @return true if the provided Element is the correct type, otherwise false.
 	 */
-	UFUNCTION(BlueprintPure = false, Category = "LearningAgents", meta = (AdvancedDisplay = 4, ReturnDisplayName = "Success", DefaultToSelf = "VisualLoggerListener"))
+	UFUNCTION(BlueprintPure = false, Category = "LearningAgents", meta = (AdvancedDisplay = 4, ReturnDisplayName = "Success", AutoCreateRefTerm = "VisualLoggerSettings"))
 	static UE_API bool GetAngleAction(
 		float& OutAngle, 
 		const ULearningAgentsActionObject* Object, 
@@ -1919,12 +1852,7 @@ public:
 		const float RelativeAngle = 0.0f, 
 		const bool bActive = true, 
 		const FName Tag = TEXT("AngleAction"),
-		const bool bVisualLoggerEnabled = false,
-		ULearningAgentsManagerListener* VisualLoggerListener = nullptr,
-		const int32 VisualLoggerAgentId = -1,
-		const FVector VisualLoggerAngleLocation = FVector::ZeroVector,
-		const FVector VisualLoggerLocation = FVector::ZeroVector,
-		const FLinearColor VisualLoggerColor = FLinearColor::Blue);
+		const FLearningAgentsVisualLoggerSettings& VisualLoggerSettings = FLearningAgentsVisualLoggerSettings());
 
 	/**
 	 * Get the value for an angle action. Returned angle is in radians.
@@ -1943,7 +1871,7 @@ public:
 	 * @param VisualLoggerColor The color for the visual logger display.
 	 * @return true if the provided Element is the correct type, otherwise false.
 	 */
-	UFUNCTION(BlueprintPure = false, Category = "LearningAgents", meta = (AdvancedDisplay = 4, ReturnDisplayName = "Success", DefaultToSelf = "VisualLoggerListener"))
+	UFUNCTION(BlueprintPure = false, Category = "LearningAgents", meta = (AdvancedDisplay = 4, ReturnDisplayName = "Success", AutoCreateRefTerm = "VisualLoggerSettings"))
 	static UE_API bool GetAngleActionRadians(
 		float& OutAngle, 
 		const ULearningAgentsActionObject* Object,
@@ -1951,12 +1879,7 @@ public:
 		const float RelativeAngle = 0.0f,
 		const bool bActive = true, 
 		const FName Tag = TEXT("AngleAction"),
-		const bool bVisualLoggerEnabled = false,
-		ULearningAgentsManagerListener* VisualLoggerListener = nullptr,
-		const int32 VisualLoggerAgentId = -1,
-		const FVector VisualLoggerAngleLocation = FVector::ZeroVector,
-		const FVector VisualLoggerLocation = FVector::ZeroVector,
-		const FLinearColor VisualLoggerColor = FLinearColor::Blue);
+		const FLearningAgentsVisualLoggerSettings& VisualLoggerSettings = FLearningAgentsVisualLoggerSettings());
 
 	/**
 	 * Get the value for a velocity action.
@@ -1975,7 +1898,7 @@ public:
 	 * @param VisualLoggerColor The color for the visual logger display.
 	 * @return true if the provided Element is the correct type, otherwise false.
 	 */
-	UFUNCTION(BlueprintPure = false, Category = "LearningAgents", meta = (AdvancedDisplay = 4, ReturnDisplayName = "Success", DefaultToSelf = "VisualLoggerListener"))
+	UFUNCTION(BlueprintPure = false, Category = "LearningAgents", meta = (AdvancedDisplay = 4, ReturnDisplayName = "Success", AutoCreateRefTerm = "VisualLoggerSettings"))
 	static UE_API bool GetVelocityAction(
 		FVector& OutVelocity, 
 		const ULearningAgentsActionObject* Object, 
@@ -1983,12 +1906,7 @@ public:
 		const FTransform RelativeTransform = FTransform(), 
 		const bool bActive = true, 
 		const FName Tag = TEXT("VelocityAction"),
-		const bool bVisualLoggerEnabled = false,
-		ULearningAgentsManagerListener* VisualLoggerListener = nullptr,
-		const int32 VisualLoggerAgentId = -1,
-		const FVector VisualLoggerVelocityLocation = FVector::ZeroVector,
-		const FVector VisualLoggerLocation = FVector::ZeroVector,
-		const FLinearColor VisualLoggerColor = FLinearColor::Blue);
+		const FLearningAgentsVisualLoggerSettings& VisualLoggerSettings = FLearningAgentsVisualLoggerSettings());
 
 	/**
 	 * Get the value for a direction action.
@@ -2008,7 +1926,7 @@ public:
 	 * @param VisualLoggerColor The color for the visual logger display.
 	 * @return true if the provided Element is the correct type, otherwise false.
 	 */
-	UFUNCTION(BlueprintPure = false, Category = "LearningAgents", meta = (AdvancedDisplay = 4, ReturnDisplayName = "Success", DefaultToSelf = "VisualLoggerListener"))
+	UFUNCTION(BlueprintPure = false, Category = "LearningAgents", meta = (AdvancedDisplay = 4, ReturnDisplayName = "Success", AutoCreateRefTerm = "VisualLoggerSettings"))
 	static UE_API bool GetDirectionAction(
 		FVector& OutDirection, 
 		const ULearningAgentsActionObject* Object, 
@@ -2016,16 +1934,10 @@ public:
 		const FTransform RelativeTransform = FTransform(),
 		const bool bActive = true, 
 		const FName Tag = TEXT("DirectionAction"),
-		const bool bVisualLoggerEnabled = false,
-		ULearningAgentsManagerListener* VisualLoggerListener = nullptr,
-		const int32 VisualLoggerAgentId = -1,
-		const FVector VisualLoggerDirectionLocation = FVector::ZeroVector,
-		const FVector VisualLoggerLocation = FVector::ZeroVector,
-		const float VisualLoggerArrowLength = 100.0f,
-		const FLinearColor VisualLoggerColor = FLinearColor::Blue);
+		const FLearningAgentsVisualLoggerSettings& VisualLoggerSettings = FLearningAgentsVisualLoggerSettings());
 };
 
 
 
 
-#undef UE_API
+#undef UE_API
\ No newline at end of file
diff --git a/Engine/Plugins/Experimental/LearningAgents/Source/LearningAgents/Public/LearningAgentsNeuralNetwork.h b/Engine/Plugins/Experimental/LearningAgents/Source/LearningAgents/Public/LearningAgentsNeuralNetwork.h
--- a/Engine/Plugins/Experimental/LearningAgents/Source/LearningAgents/Public/LearningAgentsNeuralNetwork.h
+++ b/Engine/Plugins/Experimental/LearningAgents/Source/LearningAgents/Public/LearningAgentsNeuralNetwork.h
@@ -3,7 +3,8 @@
 #pragma once
 
 #include "Engine/DataAsset.h"
-
+#include "LearningLog.h"
+#include "NNERuntimeBasicCpuBuilder.h"
 #include "LearningAgentsNeuralNetwork.generated.h"
 
 #define UE_API LEARNINGAGENTS_API
@@ -25,6 +26,48 @@ enum class ELearningAgentsActivationFunction : uint8
 	GELU	UMETA(DisplayName = "GELU"),
 };
 
+namespace UE::Learning::Agents::NeuralNetwork
+{
+	inline UE::NNE::RuntimeBasic::FModelBuilder::EActivationFunction GetBuilderActivationFunction(const ELearningAgentsActivationFunction ActivationFunction)
+	{
+		switch (ActivationFunction)
+		{
+		case ELearningAgentsActivationFunction::ReLU: return UE::NNE::RuntimeBasic::FModelBuilder::EActivationFunction::ReLU;
+		case ELearningAgentsActivationFunction::ELU: return UE::NNE::RuntimeBasic::FModelBuilder::EActivationFunction::ELU;
+		case ELearningAgentsActivationFunction::TanH: return UE::NNE::RuntimeBasic::FModelBuilder::EActivationFunction::TanH;
+		case ELearningAgentsActivationFunction::GELU: return UE::NNE::RuntimeBasic::FModelBuilder::EActivationFunction::GELU;
+		default:
+			UE_LOG(LogLearning, Error, TEXT("Unknown Activation Function"));
+			return UE::NNE::RuntimeBasic::FModelBuilder::EActivationFunction::ReLU;
+		}
+	}
+}
+
+/** Normalization options for neural network inputs & outputs. */
+UENUM(BlueprintType)
+enum class ELearningAgentsNormalization : uint8
+{
+	/** 
+	 * Uses user specified scales and biases to normalize/denormalize network dataflow at the cpp schema layer.
+	 * <IMPORTANT> Hand picked scales and offsets values will only be used if normalization is set to "Manual."
+	 */
+	Manual = 0,
+
+	/**
+	 * Uses dataset derived standard deviation and mean as scales and biases to normalize/denormalize network dataflow.
+	 * If a feature has multiple dimensions, a single scale and bias will be computed for that feature.
+	 * <IMPORTANT> Hand picked scales and offsets values will NOT be used if normalization is set to "AutoShared."
+	 */
+	AutoShared = 1,
+
+	/**
+	 * Uses dataset derived standard deviation and mean as scales and biases to normalize/denormalize network dataflow.
+	 * If a feature has multiple dimensions, a scale and bias will be computed for each dimension.
+	 * <IMPORTANT> Hand picked scales and offsets values will NOT be used if normalization is set to "AutoPerDimension."
+	 */
+	AutoPerDimension = 2,
+};
+
 class ULearningNeuralNetworkData;
 
 /** A neural network data asset. */
