[LearningAgents] Flow Matching and Normalizations

diff --git a/Engine/Plugins/Experimental/LearningAgents/Source/LearningAgents/Public/LearningAgentsObservations.h b/Engine/Plugins/Experimental/LearningAgents/Source/LearningAgents/Public/LearningAgentsObservations.h
--- a/Engine/Plugins/Experimental/LearningAgents/Source/LearningAgents/Public/LearningAgentsObservations.h
+++ b/Engine/Plugins/Experimental/LearningAgents/Source/LearningAgents/Public/LearningAgentsObservations.h
@@ -41,6 +41,98 @@ struct FLearningAgentsObservationObjectElement
 	UE::Learning::Observation::FObjectElement ObjectElement;
 };
 
+/** Settings for normalizing observations */
+USTRUCT(BlueprintType)
+struct FLearningAgentsObservationNormalizationSettings
+{
+	GENERATED_BODY()
+
+	/** The normalization method to use. */
+	UPROPERTY(EditAnywhere, BlueprintReadWrite, Category = "LearningAgents")
+	ELearningAgentsNormalization NormalizationOption = ELearningAgentsNormalization::Manual;
+
+	/** Per-element standard deviation values used for normalization. If empty, no scaling is applied. */
+	UPROPERTY(EditAnywhere, BlueprintReadWrite, Category = "LearningAgents", meta = (EditCondition = "NormalizationOption == ELearningAgentsNormalization::Manual"))
+	TArray<float> Scales;
+
+	/** Per-element mean values used for normalization. If empty, no bias is applied. */
+	UPROPERTY(EditAnywhere, BlueprintReadWrite, Category = "LearningAgents", meta = (EditCondition = "NormalizationOption == ELearningAgentsNormalization::Manual"))
+	TArray<float> Offsets;
+};
+
+/** Settings for normalizing observations */
+USTRUCT(BlueprintType)
+struct FLearningAgentsFloatObservationNormalizationSettings
+{
+	GENERATED_BODY()
+
+	/** The normalization method to use. */
+	UPROPERTY(EditAnywhere, BlueprintReadWrite, Category = "LearningAgents")
+	ELearningAgentsNormalization NormalizationOption = ELearningAgentsNormalization::Manual;
+
+	/** Per-element standard deviation values used for normalization. If empty, no scaling is applied. */
+	UPROPERTY(EditAnywhere, BlueprintReadWrite, Category = "LearningAgents", meta = (EditCondition = "NormalizationOption == ELearningAgentsNormalization::Manual"))
+	float Scale = 1.0f;
+
+	/** Per-element mean values used for normalization. If empty, no bias is applied. */
+	UPROPERTY(EditAnywhere, BlueprintReadWrite, Category = "LearningAgents", meta = (EditCondition = "NormalizationOption == ELearningAgentsNormalization::Manual"))
+	float Offset = 0.0f;
+};
+
+/** Settings for normalizing observations */
+USTRUCT(BlueprintType)
+struct FLearningAgentsVectorObservationNormalizationSettings
+{
+	GENERATED_BODY()
+
+	/** The normalization method to use. */
+	UPROPERTY(EditAnywhere, BlueprintReadWrite, Category = "LearningAgents")
+	ELearningAgentsNormalization NormalizationOption = ELearningAgentsNormalization::Manual;
+
+	/** Per-element standard deviation values used for normalization. If empty, no scaling is applied. */
+	UPROPERTY(EditAnywhere, BlueprintReadWrite, Category = "LearningAgents", meta = (EditCondition = "NormalizationOption == ELearningAgentsNormalization::Manual"))
+	FVector Scales = FVector::OneVector;
+
+	/** Per-element mean values used for normalization. If empty, no bias is applied. */
+	UPROPERTY(EditAnywhere, BlueprintReadWrite, Category = "LearningAgents", meta = (EditCondition = "NormalizationOption == ELearningAgentsNormalization::Manual"))
+	FVector Offsets = FVector::ZeroVector;
+};
+
+/** Settings for visual logger output */
+USTRUCT(BlueprintType)
+struct FLearningAgentsObservationVisualLoggerSettings
+{
+	GENERATED_BODY()
+
+	/** Whether visual logging is enabled. */
+	UPROPERTY(EditAnywhere, BlueprintReadWrite, Category = "LearningAgents")
+	bool bEnabled = false;
+
+	/** The listener object which is making this observation. This must be set to use logging. */
+	UPROPERTY(EditAnywhere, BlueprintReadWrite, Category = "LearningAgents")
+	TObjectPtr<ULearningAgentsManagerListener> Listener = nullptr;
+
+	/** The agent id associated with this observation. */
+	UPROPERTY(EditAnywhere, BlueprintReadWrite, Category = "LearningAgents")
+	int32 AgentId = -1;
+
+	/** An additional location for visualizing a specific observation element (e.g., rotation arrow origin, velocity arrow start). */
+	UPROPERTY(EditAnywhere, BlueprintReadWrite, Category = "LearningAgents")
+	FVector DebugLocation = FVector::ZeroVector;
+
+	/** A location for the visual logger text information in the world. */
+	UPROPERTY(EditAnywhere, BlueprintReadWrite, Category = "LearningAgents")
+	FVector Location = FVector::ZeroVector;
+
+	/** The color for the visual logger display. */
+	UPROPERTY(EditAnywhere, BlueprintReadWrite, Category = "LearningAgents")
+	FLinearColor Color = FLinearColor::Red;
+
+	/** The length of debug arrows to display for direction/velocity observations. */
+	UPROPERTY(EditAnywhere, BlueprintReadWrite, Category = "LearningAgents")
+	float ArrowLength = 100.0f;
+};
+
 /** Comparison and Hashing operators for Observation Elements */
 
 LEARNINGAGENTS_API bool operator==(const FLearningAgentsObservationSchemaElement& Lhs, const FLearningAgentsObservationSchemaElement& Rhs);
@@ -299,10 +391,11 @@ public:
 	 * @param Scale The scale used to normalize this observation.
 	 * @param bActive Whether to enable the observation as a valid network input. If set to inactive, this observation will be considered an empty observation. This flag allows easy toggles of specified observation.
 	 * @param Tag The tag of this new observation. Used during observation object validation and debugging.
+	 * @param NormalizationSettings Optional normalization settings containing per-element Scale, Bias, and whether to apply normalization.
 	 * @return The newly created observation schema element.
 	 */
-	UFUNCTION(BlueprintPure, Category = "LearningAgents", meta = (AdvancedDisplay = 3))
-	static UE_API FLearningAgentsObservationSchemaElement SpecifyContinuousObservation(ULearningAgentsObservationSchema* Schema, const int32 Size, const float Scale = 1.0f, const bool bActive = true, const FName Tag = TEXT("ContinuousObservation"));
+	UFUNCTION(BlueprintPure, Category = "LearningAgents", meta = (AdvancedDisplay = 3, AutoCreateRefTerm = "NormalizationSettings"))
+	static UE_API FLearningAgentsObservationSchemaElement SpecifyContinuousObservation(ULearningAgentsObservationSchema* Schema, const int32 Size, const FLearningAgentsObservationNormalizationSettings& NormalizationSettings = FLearningAgentsObservationNormalizationSettings(), const bool bActive = true, const FName Tag = TEXT("ContinuousObservation"));
 
 	/**
 	 * Specifies a new exclusive discrete observation. This represents a discrete observation which is an exclusive selection from multiple choices.
@@ -712,25 +805,25 @@ public:
 	 * type-specific observation does not exist.
 	 *
 	 * @param Schema The Observation Schema
-	 * @param FloatScale The scale used to normalize this observation.
+	 * @param NormalizationSettings The normalization used to normalize this observation.
 	 * @param bActive Whether to enable the observation as a valid network input. If set to inactive, this observation will be considered an empty observation. This flag allows easy toggles of specified observation.
 	 * @param Tag The tag of this new observation. Used during observation object validation and debugging.
 	 * @return The newly created observation schema element.
 	 */
-	UFUNCTION(BlueprintPure, Category = "LearningAgents", meta = (AdvancedDisplay = 2))
-	static UE_API FLearningAgentsObservationSchemaElement SpecifyFloatObservation(ULearningAgentsObservationSchema* Schema, const float FloatScale = 1.0f, const bool bActive = true, const FName Tag = TEXT("FloatObservation"));
+	UFUNCTION(BlueprintPure, Category = "LearningAgents", meta = (AdvancedDisplay = 2, AutoCreateRefTerm = "NormalizationSettings"))
+	static UE_API FLearningAgentsObservationSchemaElement SpecifyFloatObservation(ULearningAgentsObservationSchema* Schema, const FLearningAgentsFloatObservationNormalizationSettings& NormalizationSettings = FLearningAgentsFloatObservationNormalizationSettings(), const bool bActive = true, const FName Tag = TEXT("FloatObservation"));
 
 	/**
 	 * Specifies a new location observation. Allows an agent to observe the location of some entity.
 	 *
 	 * @param Schema The Observation Schema
-	 * @param LocationScale The scale used to normalize this observation in cm.
+	 * @param NormalizationSettings The normalization used to normalize this observation in cm.
 	 * @param bActive Whether to enable the observation as a valid network input. If set to inactive, this observation will be considered an empty observation. This flag allows easy toggles of specified observation.
 	 * @param Tag The tag of this new observation. Used during observation object validation and debugging.
 	 * @return The newly created observation schema element.
 	 */
-	UFUNCTION(BlueprintPure, Category = "LearningAgents", meta = (AdvancedDisplay = 2))
-	static UE_API FLearningAgentsObservationSchemaElement SpecifyLocationObservation(ULearningAgentsObservationSchema* Schema, const float LocationScale = 100.0f, const bool bActive = true, const FName Tag = TEXT("LocationObservation"));
+	UFUNCTION(BlueprintPure, Category = "LearningAgents", meta = (AdvancedDisplay = 2, AutoCreateRefTerm = "NormalizationSettings", NormalizationSettings = "(Scales=(X=100,Y=100,Z=100))"))
+	static UE_API FLearningAgentsObservationSchemaElement SpecifyLocationObservation(ULearningAgentsObservationSchema* Schema, const FLearningAgentsVectorObservationNormalizationSettings& NormalizationSettings = FLearningAgentsVectorObservationNormalizationSettings(), const bool bActive = true, const FName Tag = TEXT("LocationObservation"));
 
 	/**
 	 * Specifies a new rotation observation. Allows an agent to observe the rotation of some entity. Rotations are encoded as two columns of the
@@ -760,13 +853,13 @@ public:
 	 * Specifies a new transform observation. Allows an agent to observe the transform of some entity.
 	 *
 	 * @param Schema The Observation Schema
-	 * @param LocationScale The scale used to normalize the Location part of this transform observation in cm.
+	 * @param NormalizationSettings The normalization used to normalize the Location part of this transform observation in cm.
 	 * @param bActive Whether to enable the observation as a valid network input. If set to inactive, this observation will be considered an empty observation. This flag allows easy toggles of specified observation.
 	 * @param Tag The tag of this new observation. Used during observation object validation and debugging.
 	 * @return The newly created observation schema element.
 	 */
-	UFUNCTION(BlueprintPure, Category = "LearningAgents", meta = (AdvancedDisplay = 2))
-	static UE_API FLearningAgentsObservationSchemaElement SpecifyTransformObservation(ULearningAgentsObservationSchema* Schema, const float LocationScale = 100.0f, const bool bActive = true, const FName Tag = TEXT("TransformObservation"));
+	UFUNCTION(BlueprintPure, Category = "LearningAgents", meta = (AdvancedDisplay = 2, AutoCreateRefTerm = "NormalizationSettings", NormalizationSettings = "(Scales=(X=100,Y=100,Z=100))"))
+	static UE_API FLearningAgentsObservationSchemaElement SpecifyTransformObservation(ULearningAgentsObservationSchema* Schema, const FLearningAgentsVectorObservationNormalizationSettings& NormalizationSettings = FLearningAgentsVectorObservationNormalizationSettings(), const bool bActive = true, const FName Tag = TEXT("TransformObservation"));
 
 	/**
 	 * Specifies a new angle observation. This will be encoded as a 2-dimension Cartesian vector so that 0 and 350 are close to each other in the
@@ -775,7 +868,7 @@ public:
 	 * @param Schema The Observation Schema
 	 * @param bActive Whether to enable the observation as a valid network input. If set to inactive, this observation will be considered an empty observation. This flag allows easy toggles of specified observation.
 	 * @param Tag The tag of this new observation. Used during observation object validation and debugging.
-	 * @return The newly created observation schema element.
+	 * @return The newly created observation schema element.SpecifyAngleObservation
 	 */
 	UFUNCTION(BlueprintPure, Category = "LearningAgents", meta = (AdvancedDisplay = 1))
 	static UE_API FLearningAgentsObservationSchemaElement SpecifyAngleObservation(ULearningAgentsObservationSchema* Schema, const bool bActive = true, const FName Tag = TEXT("AngleObservation"));
@@ -789,8 +882,8 @@ public:
 	 * @param Tag The tag of this new observation. Used during observation object validation and debugging.
 	 * @return The newly created observation schema element.
 	 */
-	UFUNCTION(BlueprintPure, Category = "LearningAgents", meta = (AdvancedDisplay = 2))
-	static UE_API FLearningAgentsObservationSchemaElement SpecifyVelocityObservation(ULearningAgentsObservationSchema* Schema, const float VelocityScale = 200.0f, const bool bActive = true, const FName Tag = TEXT("VelocityObservation"));
+	UFUNCTION(BlueprintPure, Category = "LearningAgents", meta = (AdvancedDisplay = 2, AutoCreateRefTerm = "NormalizationSettings", NormalizationSettings = "(Scales=(X=200,Y=200,Z=200))"))
+	static UE_API FLearningAgentsObservationSchemaElement SpecifyVelocityObservation(ULearningAgentsObservationSchema* Schema, const FLearningAgentsVectorObservationNormalizationSettings& NormalizationSettings = FLearningAgentsVectorObservationNormalizationSettings(), const bool bActive = true, const FName Tag = TEXT("VelocityObservation"));
 
 	/**
 	 * Specifies a new direction observation. Allows an agent to observe the direction of some entity.
@@ -810,10 +903,11 @@ public:
 	 * @param LocationScale The scale used to normalize this observation in cm.
 	 * @param bActive Whether to enable the observation as a valid network input. If set to inactive, this observation will be considered an empty observation. This flag allows easy toggles of specified observation.
 	 * @param Tag The tag of this new observation. Used during observation object validation and debugging.
+	 * @param NormalizationSettings Optional normalization settings containing per-element Scales, Biases, and whether to apply normalization.
 	 * @return The newly created observation schema element.
 	 */
-	UFUNCTION(BlueprintPure, Category = "LearningAgents", meta = (AdvancedDisplay = 2))
-	static UE_API FLearningAgentsObservationSchemaElement SpecifyLocationAlongSplineObservation(ULearningAgentsObservationSchema* Schema, const float LocationScale = 100.0f, const bool bActive = true, const FName Tag = TEXT("LocationAlongSplineObservation"));
+	UFUNCTION(BlueprintPure, Category = "LearningAgents", meta = (AdvancedDisplay = 2, AutoCreateRefTerm = "NormalizationSettings", NormalizationSettings = "(Scales=(X=100,Y=100,Z=100))"))
+	static UE_API FLearningAgentsObservationSchemaElement SpecifyLocationAlongSplineObservation(ULearningAgentsObservationSchema* Schema, const FLearningAgentsVectorObservationNormalizationSettings& NormalizationSettings = FLearningAgentsVectorObservationNormalizationSettings(), const bool bActive = true, const FName Tag = TEXT("LocationAlongSplineObservation"));
 
 	/**
 	 * Specifies a new proportion along spline observation. This observes the proportion along a spline at the given distance. For looped splines
@@ -889,17 +983,13 @@ public:
 	 * @param VisualLoggerColor The color for the visual logger display.
 	 * @return The newly created observation object element.
 	 */
-	UFUNCTION(BlueprintPure, Category = "LearningAgents", meta = (AdvancedDisplay = 2, DefaultToSelf = "VisualLoggerListener"))
+	UFUNCTION(BlueprintPure, Category = "LearningAgents", meta = (AdvancedDisplay = 2, AutoCreateRefTerm = "VisualLoggerSettings"))
 	static UE_API FLearningAgentsObservationObjectElement MakeContinuousObservation(
 		ULearningAgentsObservationObject* Object,
 		const TArray<float>& Values,
 		const bool bActive = true,
 		const FName Tag = TEXT("ContinuousObservation"),
-		const bool bVisualLoggerEnabled = false,
-		ULearningAgentsManagerListener* VisualLoggerListener = nullptr,
-		const int32 VisualLoggerAgentId = -1,
-		const FVector VisualLoggerLocation = FVector::ZeroVector,
-		const FLinearColor VisualLoggerColor = FLinearColor::Red);
+		const FLearningAgentsObservationVisualLoggerSettings& VisualLoggerSettings = FLearningAgentsObservationVisualLoggerSettings());
 
 	/**
 	 * Make a new continuous observation. The size of Values must match the Size given during Specify.
@@ -916,15 +1006,11 @@ public:
 	 * @return The newly created observation object element.
 	 */
 	static UE_API FLearningAgentsObservationObjectElement MakeContinuousObservationFromArrayView(
-		ULearningAgentsObservationObject* Object, 
-		const TArrayView<const float> Values, 
-		const bool bActive = true, 
+		ULearningAgentsObservationObject* Object,
+		const TArrayView<const float> Values,
+		const bool bActive = true,
 		const FName Tag = TEXT("ContinuousObservation"),
-		const bool bVisualLoggerEnabled = false,
-		ULearningAgentsManagerListener* VisualLoggerListener = nullptr,
-		const int32 VisualLoggerAgentId = -1,
-		const FVector VisualLoggerLocation = FVector::ZeroVector,
-		const FLinearColor VisualLoggerColor = FLinearColor::Red);
+		const FLearningAgentsObservationVisualLoggerSettings& VisualLoggerSettings = FLearningAgentsObservationVisualLoggerSettings());
 
 	/**
 	 * Make a new exclusive discrete observation.
@@ -940,17 +1026,13 @@ public:
 	 * @param VisualLoggerColor The color for the visual logger display.
 	 * @return The newly created observation object element.
 	 */
-	UFUNCTION(BlueprintPure, Category = "LearningAgents", meta = (AdvancedDisplay = 2, DefaultToSelf = "VisualLoggerListener"))
+	UFUNCTION(BlueprintPure, Category = "LearningAgents", meta = (AdvancedDisplay = 2, AutoCreateRefTerm = "VisualLoggerSettings"))
 	static UE_API FLearningAgentsObservationObjectElement MakeExclusiveDiscreteObservation(
 		ULearningAgentsObservationObject* Object, 
 		const int32 DiscreteIndex, 
 		const bool bActive = true, 
 		const FName Tag = TEXT("ExclusiveDiscreteObservation"),
-		const bool bVisualLoggerEnabled = false,
-		ULearningAgentsManagerListener* VisualLoggerListener = nullptr,
-		const int32 VisualLoggerAgentId = -1,
-		const FVector VisualLoggerLocation = FVector::ZeroVector,
-		const FLinearColor VisualLoggerColor = FLinearColor::Red);
+		const FLearningAgentsObservationVisualLoggerSettings& VisualLoggerSettings = FLearningAgentsObservationVisualLoggerSettings());
 
 	/**
 	 * Make a new inclusive discrete observation.
@@ -966,17 +1048,13 @@ public:
 	 * @param VisualLoggerColor The color for the visual logger display.
 	 * @return The newly created observation object element.
 	 */
-	UFUNCTION(BlueprintPure, Category = "LearningAgents", meta = (AdvancedDisplay = 2, DefaultToSelf = "VisualLoggerListener"))
+	UFUNCTION(BlueprintPure, Category = "LearningAgents", meta = (AdvancedDisplay = 2, AutoCreateRefTerm = "VisualLoggerSettings"))
 	static UE_API FLearningAgentsObservationObjectElement MakeInclusiveDiscreteObservation(
 		ULearningAgentsObservationObject* Object, 
 		const TArray<int32>& DiscreteIndices, 
 		const bool bActive = true, 
 		const FName Tag = TEXT("InclusiveDiscreteObservation"),
-		const bool bVisualLoggerEnabled = false,
-		ULearningAgentsManagerListener* VisualLoggerListener = nullptr,
-		const int32 VisualLoggerAgentId = -1,
-		const FVector VisualLoggerLocation = FVector::ZeroVector,
-		const FLinearColor VisualLoggerColor = FLinearColor::Red);
+		const FLearningAgentsObservationVisualLoggerSettings& VisualLoggerSettings = FLearningAgentsObservationVisualLoggerSettings());
 	
 	/**
 	 * Make a new inclusive discrete observation.
@@ -997,11 +1075,7 @@ public:
 		const TArrayView<const int32> DiscreteIndices, 
 		const bool bActive = true, 
 		const FName Tag = TEXT("InclusiveDiscreteObservation"),
-		const bool bVisualLoggerEnabled = false,
-		ULearningAgentsManagerListener* VisualLoggerListener = nullptr,
-		const int32 VisualLoggerAgentId = -1,
-		const FVector VisualLoggerLocation = FVector::ZeroVector,
-		const FLinearColor VisualLoggerColor = FLinearColor::Red);
+		const FLearningAgentsObservationVisualLoggerSettings& VisualLoggerSettings = FLearningAgentsObservationVisualLoggerSettings());
 
 	/**
 	 * Make a new named exclusive discrete observation.
@@ -1017,17 +1091,13 @@ public:
 	 * @param VisualLoggerColor The color for the visual logger display.
 	 * @return The newly created observation object element.
 	 */
-	UFUNCTION(BlueprintPure, Category = "LearningAgents", meta = (AdvancedDisplay = 2, DefaultToSelf = "VisualLoggerListener"))
+	UFUNCTION(BlueprintPure, Category = "LearningAgents", meta = (AdvancedDisplay = 2, AutoCreateRefTerm = "VisualLoggerSettings"))
 	static UE_API FLearningAgentsObservationObjectElement MakeNamedExclusiveDiscreteObservation(
 		ULearningAgentsObservationObject* Object,
 		const FName ElementName,
 		const bool bActive = true, 
 		const FName Tag = TEXT("NamedExclusiveDiscreteObservation"),
-		const bool bVisualLoggerEnabled = false,
-		ULearningAgentsManagerListener* VisualLoggerListener = nullptr,
-		const int32 VisualLoggerAgentId = -1,
-		const FVector VisualLoggerLocation = FVector::ZeroVector,
-		const FLinearColor VisualLoggerColor = FLinearColor::Red);
+		const FLearningAgentsObservationVisualLoggerSettings& VisualLoggerSettings = FLearningAgentsObservationVisualLoggerSettings());
 
 	/**
 	 * Make a new named inclusive discrete observation.
@@ -1043,17 +1113,13 @@ public:
 	 * @param VisualLoggerColor The color for the visual logger display.
 	 * @return The newly created observation object element.
 	 */
-	UFUNCTION(BlueprintPure, Category = "LearningAgents", meta = (AdvancedDisplay = 2, DefaultToSelf = "VisualLoggerListener"))
+	UFUNCTION(BlueprintPure, Category = "LearningAgents", meta = (AdvancedDisplay = 2, AutoCreateRefTerm = "VisualLoggerSettings"))
 	static UE_API FLearningAgentsObservationObjectElement MakeNamedInclusiveDiscreteObservation(
 		ULearningAgentsObservationObject* Object,
 		const TArray<FName>& ElementNames,
 		const bool bActive = true, 
 		const FName Tag = TEXT("NamedInclusiveDiscreteObservation"),
-		const bool bVisualLoggerEnabled = false,
-		ULearningAgentsManagerListener* VisualLoggerListener = nullptr,
-		const int32 VisualLoggerAgentId = -1,
-		const FVector VisualLoggerLocation = FVector::ZeroVector,
-		const FLinearColor VisualLoggerColor = FLinearColor::Red);
+		const FLearningAgentsObservationVisualLoggerSettings& VisualLoggerSettings = FLearningAgentsObservationVisualLoggerSettings());
 
 	/**
 	 * Make a new named inclusive discrete observation.
@@ -1074,11 +1140,7 @@ public:
 		const TArrayView<const FName> ElementNames,
 		const bool bActive = true, 
 		const FName Tag = TEXT("NamedInclusiveDiscreteObservation"),
-		const bool bVisualLoggerEnabled = false,
-		ULearningAgentsManagerListener* VisualLoggerListener = nullptr,
-		const int32 VisualLoggerAgentId = -1,
-		const FVector VisualLoggerLocation = FVector::ZeroVector,
-		const FLinearColor VisualLoggerColor = FLinearColor::Red);
+		const FLearningAgentsObservationVisualLoggerSettings& VisualLoggerSettings = FLearningAgentsObservationVisualLoggerSettings());
 
 	/**
 	 * Make a new count observation.
@@ -1095,18 +1157,14 @@ public:
 	 * @param VisualLoggerColor The color for the visual logger display.
 	 * @return The newly created observation object element.
 	 */
-	UFUNCTION(BlueprintPure, Category = "LearningAgents", meta = (AdvancedDisplay = 3, DefaultToSelf = "VisualLoggerListener"))
+	UFUNCTION(BlueprintPure, Category = "LearningAgents", meta = (AdvancedDisplay = 3, AutoCreateRefTerm = "VisualLoggerSettings"))
 	static UE_API FLearningAgentsObservationObjectElement MakeCountObservation(
 		ULearningAgentsObservationObject* Object, 
 		const int32 Num, 
 		const int32 MaxNum, 
 		const bool bActive = true, 
 		const FName Tag = TEXT("CountObservation"),
-		const bool bVisualLoggerEnabled = false,
-		ULearningAgentsManagerListener* VisualLoggerListener = nullptr,
-		const int32 VisualLoggerAgentId = -1,
-		const FVector VisualLoggerLocation = FVector::ZeroVector,
-		const FLinearColor VisualLoggerColor = FLinearColor::Red);
+		const FLearningAgentsObservationVisualLoggerSettings& VisualLoggerSettings = FLearningAgentsObservationVisualLoggerSettings());
 
 	/**
 	 * Make a new struct observation.
@@ -1341,18 +1399,14 @@ public:
 	 * @param VisualLoggerColor The color for the visual logger display.
 	 * @return The newly created observation object element.
 	 */
-	UFUNCTION(BlueprintPure, Category = "LearningAgents", meta = (AdvancedDisplay = 3, DefaultToSelf = "VisualLoggerListener"))
+	UFUNCTION(BlueprintPure, Category = "LearningAgents", meta = (AdvancedDisplay = 3, AutoCreateRefTerm = "VisualLoggerSettings"))
 	static UE_API FLearningAgentsObservationObjectElement MakeEnumObservation(
 		ULearningAgentsObservationObject* Object, 
 		const UEnum* Enum, 
 		const uint8 EnumValue, 
 		const bool bActive = true, 
 		const FName Tag = TEXT("EnumObservation"),
-		const bool bVisualLoggerEnabled = false,
-		ULearningAgentsManagerListener* VisualLoggerListener = nullptr,
-		const int32 VisualLoggerAgentId = -1,
-		const FVector VisualLoggerLocation = FVector::ZeroVector,
-		const FLinearColor VisualLoggerColor = FLinearColor::Red);
+		const FLearningAgentsObservationVisualLoggerSettings& VisualLoggerSettings = FLearningAgentsObservationVisualLoggerSettings());
 
 	/**
 	 * Make a new bitmask observation.
@@ -1369,18 +1423,14 @@ public:
 	 * @param VisualLoggerColor The color for the visual logger display.
 	 * @return The newly created observation object element.
 	 */
-	UFUNCTION(BlueprintPure, Category = "LearningAgents", meta = (AdvancedDisplay = 3, DefaultToSelf = "VisualLoggerListener"))
+	UFUNCTION(BlueprintPure, Category = "LearningAgents", meta = (AdvancedDisplay = 3, AutoCreateRefTerm = "VisualLoggerSettings"))
 	static UE_API FLearningAgentsObservationObjectElement MakeBitmaskObservation(
 		ULearningAgentsObservationObject* Object, 
 		const UEnum* Enum, 
 		const int32 BitmaskValue, 
 		const bool bActive = true, 
 		const FName Tag = TEXT("BitmaskObservation"),
-		const bool bVisualLoggerEnabled = false,
-		ULearningAgentsManagerListener* VisualLoggerListener = nullptr,
-		const int32 VisualLoggerAgentId = -1,
-		const FVector VisualLoggerLocation = FVector::ZeroVector,
-		const FLinearColor VisualLoggerColor = FLinearColor::Red);
+		const FLearningAgentsObservationVisualLoggerSettings& VisualLoggerSettings = FLearningAgentsObservationVisualLoggerSettings());
 
 	/**
 	 * Make a new optional observation.
@@ -1481,17 +1531,13 @@ public:
 	 * @param VisualLoggerColor The color for the visual logger display.
 	 * @return The newly created observation object element.
 	 */
-	UFUNCTION(BlueprintPure, Category = "LearningAgents", meta = (AdvancedDisplay = 2, DefaultToSelf = "VisualLoggerListener"))
+	UFUNCTION(BlueprintPure, Category = "LearningAgents", meta = (AdvancedDisplay = 2, AutoCreateRefTerm = "VisualLoggerSettings"))
 	static UE_API FLearningAgentsObservationObjectElement MakeBoolObservation(
 		ULearningAgentsObservationObject* Object,
 		const bool bValue,
 		const bool bActive = true, 
 		const FName Tag = TEXT("BoolObservation"),
-		const bool bVisualLoggerEnabled = false,
-		ULearningAgentsManagerListener* VisualLoggerListener = nullptr,
-		const int32 VisualLoggerAgentId = -1,
-		const FVector VisualLoggerLocation = FVector::ZeroVector,
-		const FLinearColor VisualLoggerColor = FLinearColor::Red);
+		const FLearningAgentsObservationVisualLoggerSettings& VisualLoggerSettings = FLearningAgentsObservationVisualLoggerSettings());
 
 	/**
 	 * Make a new 1d convolution observation.
@@ -1502,7 +1548,7 @@ public:
 	 * @param Tag The tag of the corresponding observation. Must match the tag given during Specify.
 	 * @return The newly created observation object element.
 	 */
-	UFUNCTION(BlueprintPure, Category = "LearningAgents", meta = (AdvancedDisplay = 2, DefaultToSelf = "VisualLoggerListener"))
+	UFUNCTION(BlueprintPure, Category = "LearningAgents", meta = (AdvancedDisplay = 2, AutoCreateRefTerm = "VisualLoggerSettings"))
 	static UE_API FLearningAgentsObservationObjectElement MakeConv1dObservation(ULearningAgentsObservationObject* Object, const FLearningAgentsObservationObjectElement Element, const bool bActive = true, const FName Tag = TEXT("Conv1dObservation"));
 
 	/**
@@ -1514,7 +1560,7 @@ public:
 	 * @param Tag The tag of the corresponding observation. Must match the tag given during Specify.
 	 * @return The newly created observation object element.
 	 */
-	UFUNCTION(BlueprintPure, Category = "LearningAgents", meta = (AdvancedDisplay = 2, DefaultToSelf = "VisualLoggerListener"))
+	UFUNCTION(BlueprintPure, Category = "LearningAgents", meta = (AdvancedDisplay = 2, AutoCreateRefTerm = "VisualLoggerSettings"))
 	static UE_API FLearningAgentsObservationObjectElement MakeConv2dObservation(ULearningAgentsObservationObject* Object, const FLearningAgentsObservationObjectElement Element, const bool bActive = true, const FName Tag = TEXT("Conv2dObservation"));
 
 	/**
@@ -1531,17 +1577,13 @@ public:
 	 * @param VisualLoggerColor The color for the visual logger display.
 	 * @return The newly created observation object element.
 	 */
-	UFUNCTION(BlueprintPure, Category = "LearningAgents", meta = (AdvancedDisplay = 2, DefaultToSelf = "VisualLoggerListener"))
+	UFUNCTION(BlueprintPure, Category = "LearningAgents", meta = (AdvancedDisplay = 2, AutoCreateRefTerm = "VisualLoggerSettings"))
 	static UE_API FLearningAgentsObservationObjectElement MakeFloatObservation(
 		ULearningAgentsObservationObject* Object,
 		const float Value,
-		const bool bActive = true, 
+		const bool bActive = true,
 		const FName Tag = TEXT("FloatObservation"),
-		const bool bVisualLoggerEnabled = false,
-		ULearningAgentsManagerListener* VisualLoggerListener = nullptr,
-		const int32 VisualLoggerAgentId = -1,
-		const FVector VisualLoggerLocation = FVector::ZeroVector,
-		const FLinearColor VisualLoggerColor = FLinearColor::Red);
+		const FLearningAgentsObservationVisualLoggerSettings& VisualLoggerSettings = FLearningAgentsObservationVisualLoggerSettings());
 
 	/**
 	 * Make a new location observation.
@@ -1558,18 +1600,14 @@ public:
 	 * @param VisualLoggerColor The color for the visual logger display.
 	 * @return The newly created observation object element.
 	 */
-	UFUNCTION(BlueprintPure, Category = "LearningAgents", meta = (AdvancedDisplay = 3, DefaultToSelf = "VisualLoggerListener"))
+	UFUNCTION(BlueprintPure, Category = "LearningAgents", meta = (AdvancedDisplay = 3, AutoCreateRefTerm = "VisualLoggerSettings"))
 	static UE_API FLearningAgentsObservationObjectElement MakeLocationObservation(
 		ULearningAgentsObservationObject* Object,
 		const FVector Location,
 		const FTransform RelativeTransform = FTransform(),
-		const bool bActive = true, 
+		const bool bActive = true,
 		const FName Tag = TEXT("LocationObservation"),
-		const bool bVisualLoggerEnabled = false,
-		ULearningAgentsManagerListener* VisualLoggerListener = nullptr,
-		const int32 VisualLoggerAgentId = -1,
-		const FVector VisualLoggerLocation = FVector::ZeroVector,
-		const FLinearColor VisualLoggerColor = FLinearColor::Red);
+		const FLearningAgentsObservationVisualLoggerSettings& VisualLoggerSettings = FLearningAgentsObservationVisualLoggerSettings());
 
 	/**
 	 * Make a new rotation observation.
@@ -1587,19 +1625,14 @@ public:
 	 * @param VisualLoggerColor The color for the visual logger display.
 	 * @return The newly created observation object element.
 	 */
-	UFUNCTION(BlueprintPure, Category = "LearningAgents", meta = (AdvancedDisplay = 3, DefaultToSelf = "VisualLoggerListener"))
+	UFUNCTION(BlueprintPure, Category = "LearningAgents", meta = (AdvancedDisplay = 3, AutoCreateRefTerm = "VisualLoggerSettings"))
 	static UE_API FLearningAgentsObservationObjectElement MakeRotationObservation(
 		ULearningAgentsObservationObject* Object,
 		const FRotator Rotation,
 		const FRotator RelativeRotation = FRotator::ZeroRotator,
 		const bool bActive = true, 
 		const FName Tag = TEXT("RotationObservation"),
-		const bool bVisualLoggerEnabled = false,
-		ULearningAgentsManagerListener* VisualLoggerListener = nullptr,
-		const int32 VisualLoggerAgentId = -1,
-		const FVector VisualLoggerRotationLocation = FVector::ZeroVector,
-		const FVector VisualLoggerLocation = FVector::ZeroVector,
-		const FLinearColor VisualLoggerColor = FLinearColor::Red);
+		const FLearningAgentsObservationVisualLoggerSettings& VisualLoggerSettings = FLearningAgentsObservationVisualLoggerSettings());
 
 	/**
 	 * Make a new rotation observation from a quaternion.
@@ -1617,19 +1650,14 @@ public:
 	 * @param VisualLoggerColor The color for the visual logger display.
 	 * @return The newly created observation object element.
 	 */
-	UFUNCTION(BlueprintPure, Category = "LearningAgents", meta = (AdvancedDisplay = 3, DefaultToSelf = "VisualLoggerListener"))
+	UFUNCTION(BlueprintPure, Category = "LearningAgents", meta = (AdvancedDisplay = 3, AutoCreateRefTerm = "VisualLoggerSettings"))
 	static UE_API FLearningAgentsObservationObjectElement MakeRotationObservationFromQuat(
 		ULearningAgentsObservationObject* Object,
 		const FQuat Rotation,
 		const FQuat RelativeRotation,
 		const bool bActive = true, 
 		const FName Tag = TEXT("RotationObservation"),
-		const bool bVisualLoggerEnabled = false,
-		ULearningAgentsManagerListener* VisualLoggerListener = nullptr,
-		const int32 VisualLoggerAgentId = -1,
-		const FVector VisualLoggerRotationLocation = FVector::ZeroVector,
-		const FVector VisualLoggerLocation = FVector::ZeroVector,
-		const FLinearColor VisualLoggerColor = FLinearColor::Red);
+		const FLearningAgentsObservationVisualLoggerSettings& VisualLoggerSettings = FLearningAgentsObservationVisualLoggerSettings());
 
 	/**
 	 * Make a new scale observation.
@@ -1647,19 +1675,14 @@ public:
 	 * @param VisualLoggerColor The color for the visual logger display.
 	 * @return The newly created observation object element.
 	 */
-	UFUNCTION(BlueprintPure, Category = "LearningAgents", meta = (AdvancedDisplay = 3, DefaultToSelf = "VisualLoggerListener"))
+	UFUNCTION(BlueprintPure, Category = "LearningAgents", meta = (AdvancedDisplay = 3, AutoCreateRefTerm = "VisualLoggerSettings"))
 	static UE_API FLearningAgentsObservationObjectElement MakeScaleObservation(
 		ULearningAgentsObservationObject* Object,
 		const FVector Scale,
 		const FVector RelativeScale = FVector(1, 1, 1),
 		const bool bActive = true, 
 		const FName Tag = TEXT("ScaleObservation"),
-		const bool bVisualLoggerEnabled = false,
-		ULearningAgentsManagerListener* VisualLoggerListener = nullptr,
-		const int32 VisualLoggerAgentId = -1,
-		const FVector VisualLoggerScaleLocation = FVector::ZeroVector,
-		const FVector VisualLoggerLocation = FVector::ZeroVector,
-		const FLinearColor VisualLoggerColor = FLinearColor::Red);
+		const FLearningAgentsObservationVisualLoggerSettings& VisualLoggerSettings = FLearningAgentsObservationVisualLoggerSettings());
 
 	/**
 	 * Make a new transform observation.
@@ -1676,18 +1699,14 @@ public:
 	 * @param VisualLoggerColor The color for the visual logger display.
 	 * @return The newly created observation object element.
 	 */
-	UFUNCTION(BlueprintPure, Category = "LearningAgents", meta = (AdvancedDisplay = 3, DefaultToSelf = "VisualLoggerListener"))
+	UFUNCTION(BlueprintPure, Category = "LearningAgents", meta = (AdvancedDisplay = 3, AutoCreateRefTerm = "VisualLoggerSettings"))
 	static UE_API FLearningAgentsObservationObjectElement MakeTransformObservation(
 		ULearningAgentsObservationObject* Object,
 		const FTransform Transform,
 		const FTransform RelativeTransform = FTransform(),
 		const bool bActive = true, 
 		const FName Tag = TEXT("TransformObservation"),
-		const bool bVisualLoggerEnabled = false,
-		ULearningAgentsManagerListener* VisualLoggerListener = nullptr,
-		const int32 VisualLoggerAgentId = -1,
-		const FVector VisualLoggerLocation = FVector::ZeroVector,
-		const FLinearColor VisualLoggerColor = FLinearColor::Red);
+		const FLearningAgentsObservationVisualLoggerSettings& VisualLoggerSettings = FLearningAgentsObservationVisualLoggerSettings());
 
 	/**
 	 * Make a new angle observation. Angles should be given in degrees.
@@ -1704,18 +1723,14 @@ public:
 	 * @param VisualLoggerColor The color for the visual logger display.
 	 * @return The newly created observation object element.
 	 */
-	UFUNCTION(BlueprintPure, Category = "LearningAgents", meta = (AdvancedDisplay = 3, DefaultToSelf = "VisualLoggerListener"))
+	UFUNCTION(BlueprintPure, Category = "LearningAgents", meta = (AdvancedDisplay = 3, AutoCreateRefTerm = "VisualLoggerSettings"))
 	static UE_API FLearningAgentsObservationObjectElement MakeAngleObservation(
 		ULearningAgentsObservationObject* Object,
 		const float Angle,
 		const float RelativeAngle = 0.0f,
 		const bool bActive = true,
 		const FName Tag = TEXT("AngleObservation"),
-		const bool bVisualLoggerEnabled = false,
-		ULearningAgentsManagerListener* VisualLoggerListener = nullptr,
-		const int32 VisualLoggerAgentId = -1,
-		const FVector VisualLoggerLocation = FVector::ZeroVector,
-		const FLinearColor VisualLoggerColor = FLinearColor::Red);
+		const FLearningAgentsObservationVisualLoggerSettings& VisualLoggerSettings = FLearningAgentsObservationVisualLoggerSettings());
 
 	/**
 	 * Make a new angle observation. Angles should be given in radians.
@@ -1732,18 +1747,14 @@ public:
 	 * @param VisualLoggerColor The color for the visual logger display.
 	 * @return The newly created observation object element.
 	 */
-	UFUNCTION(BlueprintPure, Category = "LearningAgents", meta = (AdvancedDisplay = 3, DefaultToSelf = "VisualLoggerListener"))
+	UFUNCTION(BlueprintPure, Category = "LearningAgents", meta = (AdvancedDisplay = 3, AutoCreateRefTerm = "VisualLoggerSettings"))
 	static UE_API FLearningAgentsObservationObjectElement MakeAngleObservationRadians(
 		ULearningAgentsObservationObject* Object, 
 		const float Angle, 
 		const float RelativeAngle = 0.0f, 
 		const bool bActive = true, 
 		const FName Tag = TEXT("AngleObservation"),
-		const bool bVisualLoggerEnabled = false,
-		ULearningAgentsManagerListener* VisualLoggerListener = nullptr,
-		const int32 VisualLoggerAgentId = -1,
-		const FVector VisualLoggerLocation = FVector::ZeroVector,
-		const FLinearColor VisualLoggerColor = FLinearColor::Red);
+		const FLearningAgentsObservationVisualLoggerSettings& VisualLoggerSettings = FLearningAgentsObservationVisualLoggerSettings());
 
 	/**
 	 * Make a new velocity observation.
@@ -1761,19 +1772,14 @@ public:
 	 * @param VisualLoggerColor The color for the visual logger display.
 	 * @return The newly created observation object element.
 	 */
-	UFUNCTION(BlueprintPure, Category = "LearningAgents", meta = (AdvancedDisplay = 3, DefaultToSelf = "VisualLoggerListener"))
+	UFUNCTION(BlueprintPure, Category = "LearningAgents", meta = (AdvancedDisplay = 3, AutoCreateRefTerm = "VisualLoggerSettings"))
 	static UE_API FLearningAgentsObservationObjectElement MakeVelocityObservation(
-		ULearningAgentsObservationObject* Object, 
-		const FVector Velocity, 
-		const FTransform RelativeTransform = FTransform(), 
-		const bool bActive = true, 
+		ULearningAgentsObservationObject* Object,
+		const FVector Velocity,
+		const FTransform RelativeTransform = FTransform(),
+		const bool bActive = true,
 		const FName Tag = TEXT("VelocityObservation"),
-		const bool bVisualLoggerEnabled = false,
-		ULearningAgentsManagerListener* VisualLoggerListener = nullptr,
-		const int32 VisualLoggerAgentId = -1,
-		const FVector VisualLoggerVelocityLocation = FVector::ZeroVector,
-		const FVector VisualLoggerLocation = FVector::ZeroVector,
-		const FLinearColor VisualLoggerColor = FLinearColor::Red);
+		const FLearningAgentsObservationVisualLoggerSettings& VisualLoggerSettings = FLearningAgentsObservationVisualLoggerSettings());
 
 	/**
 	 * Make a new direction observation.
@@ -1792,20 +1798,14 @@ public:
 	 * @param VisualLoggerColor The color for the visual logger display.
 	 * @return The newly created observation object element.
 	 */
-	UFUNCTION(BlueprintPure, Category = "LearningAgents", meta = (AdvancedDisplay = 3, DefaultToSelf = "VisualLoggerListener"))
+	UFUNCTION(BlueprintPure, Category = "LearningAgents", meta = (AdvancedDisplay = 3, AutoCreateRefTerm = "VisualLoggerSettings"))
 	static UE_API FLearningAgentsObservationObjectElement MakeDirectionObservation(
 		ULearningAgentsObservationObject* Object,
 		const FVector Direction,
 		const FTransform RelativeTransform = FTransform(),
-		const bool bActive = true, 
+		const bool bActive = true,
 		const FName Tag = TEXT("DirectionObservation"),
-		const bool bVisualLoggerEnabled = false,
-		ULearningAgentsManagerListener* VisualLoggerListener = nullptr,
-		const int32 VisualLoggerAgentId = -1,
-		const FVector VisualLoggerDirectionLocation = FVector::ZeroVector,
-		const FVector VisualLoggerLocation = FVector::ZeroVector,
-		const float VisualLoggerArrowLength = 100.0f,
-		const FLinearColor VisualLoggerColor = FLinearColor::Red);
+		const FLearningAgentsObservationVisualLoggerSettings& VisualLoggerSettings = FLearningAgentsObservationVisualLoggerSettings());
 
 	/**
 	 * Make a new location along spline observation.
@@ -1823,19 +1823,15 @@ public:
 	 * @param VisualLoggerColor The color for the visual logger display.
 	 * @return The newly created observation object element.
 	 */
-	UFUNCTION(BlueprintPure, Category = "LearningAgents", meta = (AdvancedDisplay = 4, DefaultToSelf = "VisualLoggerListener"))
+	UFUNCTION(BlueprintPure, Category = "LearningAgents", meta = (AdvancedDisplay = 4, AutoCreateRefTerm = "VisualLoggerSettings"))
 	static UE_API FLearningAgentsObservationObjectElement MakeLocationAlongSplineObservation(
 		ULearningAgentsObservationObject* Object,
 		const USplineComponent* SplineComponent,
 		const float DistanceAlongSpline,
 		const FTransform RelativeTransform = FTransform(),
-		const bool bActive = true, 
+		const bool bActive = true,
 		const FName Tag = TEXT("LocationAlongSplineObservation"),
-		const bool bVisualLoggerEnabled = false,
-		ULearningAgentsManagerListener* VisualLoggerListener = nullptr,
-		const int32 VisualLoggerAgentId = -1,
-		const FVector VisualLoggerLocation = FVector::ZeroVector,
-		const FLinearColor VisualLoggerColor = FLinearColor::Red);
+		const FLearningAgentsObservationVisualLoggerSettings& VisualLoggerSettings = FLearningAgentsObservationVisualLoggerSettings());
 
 	/**
 	 * Make a new proportion along spline observation.
@@ -1852,18 +1848,14 @@ public:
 	 * @param VisualLoggerColor The color for the visual logger display.
 	 * @return The newly created observation object element.
 	 */
-	UFUNCTION(BlueprintPure, Category = "LearningAgents", meta = (AdvancedDisplay = 3, DefaultToSelf = "VisualLoggerListener"))
+	UFUNCTION(BlueprintPure, Category = "LearningAgents", meta = (AdvancedDisplay = 3, AutoCreateRefTerm = "VisualLoggerSettings"))
 	static UE_API FLearningAgentsObservationObjectElement MakeProportionAlongSplineObservation(
 		ULearningAgentsObservationObject* Object, 
 		const USplineComponent* SplineComponent, 
 		const float DistanceAlongSpline, 
 		const bool bActive = true, 
 		const FName Tag = TEXT("ProportionAlongSplineObservation"),
-		const bool bVisualLoggerEnabled = false,
-		ULearningAgentsManagerListener* VisualLoggerListener = nullptr,
-		const int32 VisualLoggerAgentId = -1,
-		const FVector VisualLoggerLocation = FVector::ZeroVector,
-		const FLinearColor VisualLoggerColor = FLinearColor::Red);
+		const FLearningAgentsObservationVisualLoggerSettings& VisualLoggerSettings = FLearningAgentsObservationVisualLoggerSettings());
 
 	/**
 	 * Make a new direction along spline observation.
@@ -1882,7 +1874,7 @@ public:
 	 * @param VisualLoggerColor The color for the visual logger display.
 	 * @return The newly created observation object element.
 	 */
-	UFUNCTION(BlueprintPure, Category = "LearningAgents", meta = (AdvancedDisplay = 4, DefaultToSelf = "VisualLoggerListener"))
+	UFUNCTION(BlueprintPure, Category = "LearningAgents", meta = (AdvancedDisplay = 4, AutoCreateRefTerm = "VisualLoggerSettings"))
 	static UE_API FLearningAgentsObservationObjectElement MakeDirectionAlongSplineObservation(
 		ULearningAgentsObservationObject* Object,
 		const USplineComponent* SplineComponent,
@@ -1890,12 +1882,7 @@ public:
 		const FTransform RelativeTransform = FTransform(),
 		const bool bActive = true, 
 		const FName Tag = TEXT("DirectionAlongSplineObservation"),
-		const bool bVisualLoggerEnabled = false,
-		ULearningAgentsManagerListener* VisualLoggerListener = nullptr,
-		const int32 VisualLoggerAgentId = -1,
-		const FVector VisualLoggerLocation = FVector::ZeroVector,
-		const float VisualLoggerArrowLength = 100.0f,
-		const FLinearColor VisualLoggerColor = FLinearColor::Red);
+		const FLearningAgentsObservationVisualLoggerSettings& VisualLoggerSettings = FLearningAgentsObservationVisualLoggerSettings());
 
 	/**
 	 * Make a new proportion along ray observation.
@@ -1914,7 +1901,7 @@ public:
 	 * @param VisualLoggerColor The color for the visual logger display.
 	 * @return The newly created observation object element.
 	 */
-	UFUNCTION(BlueprintPure, Category = "LearningAgents", meta = (AdvancedDisplay = 5, DefaultToSelf = "VisualLoggerListener"))
+	UFUNCTION(BlueprintPure, Category = "LearningAgents", meta = (AdvancedDisplay = 5, AutoCreateRefTerm = "VisualLoggerSettings"))
 	static UE_API FLearningAgentsObservationObjectElement MakeProportionAlongRayObservation(
 		ULearningAgentsObservationObject* Object, 
 		const FVector RayStart, 
@@ -1923,11 +1910,7 @@ public:
 		const ECollisionChannel CollisionChannel = ECollisionChannel::ECC_WorldStatic, 
 		const bool bActive = true, 
 		const FName Tag = TEXT("ProportionAlongRayObservation"),
-		const bool bVisualLoggerEnabled = false,
-		ULearningAgentsManagerListener* VisualLoggerListener = nullptr,
-		const int32 VisualLoggerAgentId = -1,
-		const FVector VisualLoggerLocation = FVector::ZeroVector,
-		const FLinearColor VisualLoggerColor = FLinearColor::Red);
+		const FLearningAgentsObservationVisualLoggerSettings& VisualLoggerSettings = FLearningAgentsObservationVisualLoggerSettings());
 
 	/**
 	 * Make a new LiDAR observation.
@@ -1947,7 +1930,7 @@ public:
 	 * @param VisualLoggerColor The color for the visual logger display.
 	 * @return The newly created observation object element.
 	 */
-	UFUNCTION(BlueprintPure, Category = "LearningAgents", meta = (AdvancedDisplay = 6, DefaultToSelf = "VisualLoggerListener"))
+	UFUNCTION(BlueprintPure, Category = "LearningAgents", meta = (AdvancedDisplay = 6, AutoCreateRefTerm = "VisualLoggerSettings"))
 	static UE_API FLearningAgentsObservationObjectElement MakeLidarObservation(
 		ULearningAgentsObservationObject* Object,
 		const FVector RayStart,
@@ -1957,11 +1940,7 @@ public:
 		const AActor* ActorToIgnore = nullptr,
 		const bool bActive = true, 
 		const FName Tag = TEXT("LidarObservation"),
-		const bool bVisualLoggerEnabled = false,
-		ULearningAgentsManagerListener* VisualLoggerListener = nullptr,
-		const int32 VisualLoggerAgentId = -1,
-		const FVector VisualLoggerLocation = FVector::ZeroVector,
-		const FLinearColor VisualLoggerColor = FLinearColor::Red);
+		const FLearningAgentsObservationVisualLoggerSettings& VisualLoggerSettings = FLearningAgentsObservationVisualLoggerSettings());
 
 public:
 
diff --git a/Engine/Plugins/Experimental/LearningAgents/Source/LearningAgentsTraining/Private/LearningAgentsImitationTrainer.cpp b/Engine/Plugins/Experimental/LearningAgents/Source/LearningAgentsTraining/Private/LearningAgentsImitationTrainer.cpp
--- a/Engine/Plugins/Experimental/LearningAgents/Source/LearningAgentsTraining/Private/LearningAgentsImitationTrainer.cpp
+++ b/Engine/Plugins/Experimental/LearningAgents/Source/LearningAgentsTraining/Private/LearningAgentsImitationTrainer.cpp
@@ -24,19 +24,20 @@ TSharedRef<FJsonObject> FLearningAgentsImitationTrainerTrainingSettings::AsJsonC
 {
 	TSharedRef<FJsonObject> ConfigObject = MakeShared<FJsonObject>();
 
-	ConfigObject->SetNumberField(TEXT("IterationNum"), NumberOfIterations);
-	ConfigObject->SetNumberField(TEXT("LearningRate"), LearningRate);
-	ConfigObject->SetNumberField(TEXT("LearningRateDecay"), LearningRateDecay);
-	ConfigObject->SetNumberField(TEXT("LearningRateDecayStepSize"), LearningRateDecayStepSize);
-	ConfigObject->SetNumberField(TEXT("WeightDecay"), WeightDecay);
-	ConfigObject->SetNumberField(TEXT("BatchSize"), BatchSize);
-	ConfigObject->SetNumberField(TEXT("Window"), Window);
-	ConfigObject->SetNumberField(TEXT("ActionRegularizationWeight"), ActionRegularizationWeight);
-	ConfigObject->SetNumberField(TEXT("ActionEntropyWeight"), ActionEntropyWeight);
-	ConfigObject->SetNumberField(TEXT("ObservationNoiseScale"), ObservationNoiseScale);
-	ConfigObject->SetNumberField(TEXT("Seed"), RandomSeed);
-	ConfigObject->SetBoolField(TEXT("SaveSnapshots"), bSaveSnapshots);
-	ConfigObject->SetNumberField(TEXT("IterationsPerSnapshot"), IterationsPerSnapshot);
+	ConfigObject->SetNumberField(TEXT("IterationNum"), Hyperparameters.NumberOfIterations);
+	ConfigObject->SetNumberField(TEXT("LearningRate"), Hyperparameters.LearningRate);
+	ConfigObject->SetNumberField(TEXT("LearningRateDecay"), Hyperparameters.LearningRateDecay);
+	ConfigObject->SetNumberField(TEXT("LearningRateDecayStepSize"), Hyperparameters.LearningRateDecayStepSize);
+	ConfigObject->SetBoolField(TEXT("EmbedNormalizationsInWeights"), Hyperparameters.bEmbedNormalizationsInWeights);
+	ConfigObject->SetNumberField(TEXT("WeightDecay"), Hyperparameters.WeightDecay);
+	ConfigObject->SetNumberField(TEXT("BatchSize"), Hyperparameters.BatchSize);
+	ConfigObject->SetNumberField(TEXT("Window"), Hyperparameters.Window);
+	ConfigObject->SetNumberField(TEXT("ActionRegularizationWeight"), Hyperparameters.ActionRegularizationWeight);
+	ConfigObject->SetNumberField(TEXT("ActionEntropyWeight"), Hyperparameters.ActionEntropyWeight);
+	ConfigObject->SetNumberField(TEXT("ObservationNoiseScale"), Hyperparameters.ObservationNoiseScale);
+	ConfigObject->SetNumberField(TEXT("Seed"), Hyperparameters.RandomSeed);
+	ConfigObject->SetBoolField(TEXT("SaveSnapshots"), LoggingSettings.bSaveSnapshots);
+	ConfigObject->SetNumberField(TEXT("IterationsPerSnapshot"), LoggingSettings.IterationsPerSnapshot);
 	// Eval
 	ConfigObject->SetBoolField(TEXT("RunEvaluation"), bRunEvaluation);
 	ConfigObject->SetNumberField(TEXT("TrainEvalDatasetSplit"), TrainEvalDatasetSplit);
@@ -196,13 +197,13 @@ void ULearningAgentsImitationTrainer::BeginTraining(
 			continue;
 		}
 
-		if (!UE::Learning::Trainer::IsObservationSchemaSubsetCompatible(Recording->Schemas[Record.SchemaName].ObservationSchemaJson, Interactor->GetObservationSchema()->ObservationSchema, Interactor->GetObservationSchemaElement().SchemaElement))
+		if (ImitationTrainerTrainingSettings.ValidationSettings.bCheckObservationSchemaCompatibility && !UE::Learning::Trainer::IsObservationSchemaSubsetCompatible(Recording->Schemas[Record.SchemaName].ObservationSchemaJson, Interactor->GetObservationSchema()->ObservationSchema, Interactor->GetObservationSchemaElement().SchemaElement))
 		{
 			UE_LOG(LogLearning, Warning, TEXT("%s: Record has incompatible observation schema."), *GetName());
 			continue;
 		}
 
-		if (!UE::Learning::Trainer::IsActionSchemaSubsetCompatible(Recording->Schemas[Record.SchemaName].ActionSchemaJson, Interactor->GetActionSchema()->ActionSchema, Interactor->GetActionSchemaElement().SchemaElement))
+		if (ImitationTrainerTrainingSettings.ValidationSettings.bCheckActionSchemaCompatibility && !UE::Learning::Trainer::IsActionSchemaSubsetCompatible(Recording->Schemas[Record.SchemaName].ActionSchemaJson, Interactor->GetActionSchema()->ActionSchema, Interactor->GetActionSchemaElement().SchemaElement))
 		{
 			UE_LOG(LogLearning, Warning, TEXT("%s: Record has incompatible action schema."), *GetName());
 			continue;
@@ -251,6 +252,9 @@ void ULearningAgentsImitationTrainer::BeginTraining(
 			continue;
 		}
 
+		bool bObservationSchemaSubsetCompatible = (!ImitationTrainerTrainingSettings.ValidationSettings.bCheckObservationSchemaCompatibility) || UE::Learning::Trainer::IsObservationSchemaSubsetCompatible(Recording->Schemas[Record.SchemaName].ObservationSchemaJson, Interactor->GetObservationSchema()->ObservationSchema, Interactor->GetObservationSchemaElement().SchemaElement);
+		bool bActionSchemaSubsetCompatible = (!ImitationTrainerTrainingSettings.ValidationSettings.bCheckActionSchemaCompatibility) || UE::Learning::Trainer::IsActionSchemaSubsetCompatible(Recording->Schemas[Record.SchemaName].ActionSchemaJson, Interactor->GetActionSchema()->ActionSchema, Interactor->GetActionSchemaElement().SchemaElement);
+
 		if (Record.ObservationDimNum == ObservationNum && Record.ActionDimNum == ActionNum)
 		{
 			TLearningArrayView<2, const float> ObservationsView = TLearningArrayView<2, const float>(Record.Blob->ObservationData.GetData(), { Record.StepNum, Record.ObservationDimNum });
@@ -261,8 +265,7 @@ void ULearningAgentsImitationTrainer::BeginTraining(
 			UE::Learning::Array::Copy(RecordedObservations.Slice(StepIdx, Record.StepNum), ObservationsView);
 			UE::Learning::Array::Copy(RecordedActions.Slice(StepIdx, Record.StepNum), ActionsView);
 		}
-		else if (UE::Learning::Trainer::IsObservationSchemaSubsetCompatible(Recording->Schemas[Record.SchemaName].ObservationSchemaJson, Interactor->GetObservationSchema()->ObservationSchema, Interactor->GetObservationSchemaElement().SchemaElement) &&
-			UE::Learning::Trainer::IsActionSchemaSubsetCompatible(Recording->Schemas[Record.SchemaName].ActionSchemaJson, Interactor->GetActionSchema()->ActionSchema, Interactor->GetActionSchemaElement().SchemaElement))
+		else if (bObservationSchemaSubsetCompatible && bActionSchemaSubsetCompatible)
 		{
 			TLearningArrayView<2, const float> ObservationsView = TLearningArrayView<2, const float>(Record.Blob->ObservationData.GetData(), { Record.StepNum, Record.ObservationDimNum });
 			TLearningArrayView<2, const float> ActionsView = TLearningArrayView<2, const float>(Record.Blob->ActionData.GetData(), { Record.StepNum, Record.ActionDimNum });
@@ -379,7 +382,7 @@ void ULearningAgentsImitationTrainer::BeginTraining(
 TSharedRef<FJsonObject> ULearningAgentsImitationTrainer::CreateDataConfig() const
 {
 	TSharedRef<FJsonObject> ConfigObject = MakeShared<FJsonObject>();
-	
+
 	const int32 ObservationSchemaId = 0;
 	const int32 ActionSchemaId = 0;
 
@@ -484,9 +487,9 @@ TSharedRef<FJsonObject> ULearningAgentsImitationTrainer::CreateTrainerConfig(con
 
 	// Add Training Workflow Config Entries
 	ConfigObject->SetStringField(TEXT("Device"), UE::Learning::Trainer::GetDeviceString(UE::Learning::Agents::GetTrainingDevice(TrainingSettings.Device)));
-	ConfigObject->SetBoolField(TEXT("UseTensorBoard"), TrainingSettings.bUseTensorboard);
-	ConfigObject->SetBoolField(TEXT("UseMLflow"), TrainingSettings.bUseMLflow);
-	ConfigObject->SetStringField(TEXT("MLflowTrackingUri"), TrainingSettings.MLflowTrackingUri);
+	ConfigObject->SetBoolField(TEXT("UseTensorBoard"), TrainingSettings.LoggingSettings.bUseTensorboard);
+	ConfigObject->SetBoolField(TEXT("UseMLflow"), TrainingSettings.LoggingSettings.bUseMLflow);
+	ConfigObject->SetStringField(TEXT("MLflowTrackingUri"), TrainingSettings.LoggingSettings.MLflowTrackingUri);
 
 	return ConfigObject;
 }
@@ -515,12 +518,14 @@ void ULearningAgentsImitationTrainer::DoneTraining()
 
 	if (IsTraining())
 	{
-		// Wait for Trainer to finish
-		Trainer->Wait();
-
-		// If not finished in time, terminate
-		Trainer->Terminate();
+		if (Trainer)
+		{
+			// Wait for Trainer to finish
+			Trainer->Wait();
 
+			// If not finished in time, terminate
+			Trainer->Terminate();
+		}
 		bIsTraining = false;
 	}
 }
@@ -536,7 +541,10 @@ void ULearningAgentsImitationTrainer::EndTraining()
 	if (IsTraining())
 	{
 		UE_LOG(LogLearning, Display, TEXT("%s: Stopping training..."), *GetName());
-		Trainer->SendStop();
+		if (Trainer)
+		{
+			Trainer->SendStop();
+		}
 		DoneTraining();
 	}
 }
@@ -595,7 +603,6 @@ void ULearningAgentsImitationTrainer::IterateTraining()
 			EndTraining();
 			return;
 		}
-		Policy->GetDecoderNetworkAsset()->ForceMarkDirty();
 	}
 }
 
diff --git a/Engine/Plugins/Experimental/LearningAgents/Source/LearningAgentsTraining/Public/LearningAgentsImitationTrainer.h b/Engine/Plugins/Experimental/LearningAgents/Source/LearningAgentsTraining/Public/LearningAgentsImitationTrainer.h
--- a/Engine/Plugins/Experimental/LearningAgents/Source/LearningAgentsTraining/Public/LearningAgentsImitationTrainer.h
+++ b/Engine/Plugins/Experimental/LearningAgents/Source/LearningAgentsTraining/Public/LearningAgentsImitationTrainer.h
@@ -38,17 +38,15 @@ public:
 	/** Tag used to filter records from the recording. */
 	UPROPERTY(EditDefaultsOnly, BlueprintReadWrite, Category = "LearningAgents")
 	FGameplayTag RecordingFilterTag;
-
 };
 
-/** The configurable settings for the training process. */
+/** The hyperparameters for the imitation learning training process. */
 USTRUCT(BlueprintType, Category = "LearningAgents")
-struct FLearningAgentsImitationTrainerTrainingSettings
+struct FLearningAgentsImitationTrainerHyperparameters
 {
 	GENERATED_BODY()
 
 public:
-
 	/** The number of iterations to run before ending training. */
 	UPROPERTY(EditAnywhere, Category = "LearningAgents", meta = (ClampMin = "1", UIMin = "1"))
 	int32 NumberOfIterations = 1000000;
@@ -65,6 +63,14 @@ public:
 	UPROPERTY(EditAnywhere, Category = "LearningAgents", meta = (ClampMin = "1", UIMin = "1"))
 	int32 LearningRateDecayStepSize = 1000;
 
+	/**
+	 * Whether to embed dataset normalizations into encoder and decoder networks. 
+	 * This is recommended to set to true to stabilize trainings.
+	 * <Important> The normalization option in action/observation schemas must be set to one of the "auto" options.
+	 */
+	UPROPERTY(EditAnywhere, Category = "LearningAgents")
+	bool bEmbedNormalizationsInWeights = false;
+
 	/**
 	 * Amount of weight decay to apply to the network. Larger values encourage network weights to be smaller but too
 	 * large a value can cause the network weights to collapse to all zeros.
@@ -109,10 +115,12 @@ public:
 	/** The seed used for any random sampling the trainer will perform, e.g. for weight initialization. */
 	UPROPERTY(EditAnywhere, Category = "LearningAgents", meta = (ClampMin = "0", UIMin = "0"))
 	int32 RandomSeed = 1234;
+};
 
-	/** The device to train on. */
-	UPROPERTY(EditAnywhere, Category = "LearningAgents")
-	ELearningAgentsTrainingDevice Device = ELearningAgentsTrainingDevice::CPU;
+USTRUCT(BlueprintType, Category = "LearningAgents")
+struct FLearningAgentsImitationLearningLoggingSettings
+{
+	GENERATED_BODY()
 
 	/**
 	 * If true, TensorBoard logs will be emitted to the intermediate directory.
@@ -125,7 +133,7 @@ public:
 
 	/** If true, snapshots of the trained networks will be emitted to the intermediate directory. */
 	UPROPERTY(EditAnywhere, Category = "LearningAgents")
-	bool bSaveSnapshots = false;
+	bool bSaveSnapshots = true;
 
 	/** If bSaveSnapshots is true, the snapshots will be saved at an interval defined by the specified number of iterations. */
 	UPROPERTY(EditAnywhere, Category = "LearningAgents", meta = (ClampMin = "1", UIMin = "1"))
@@ -138,6 +146,50 @@ public:
 	/** The URI of the MLflow Tracking Server to log to. */
 	UPROPERTY(EditAnywhere, Category = "LearningAgents")
 	FString MLflowTrackingUri = "";
+};
+
+/** Validation settings for the imitation learning training process. */
+USTRUCT(BlueprintType, Category = "LearningAgents")
+struct FLearningAgentsImitationLearningValidationSettings
+{
+	GENERATED_BODY()
+
+	/** 
+	 * Whether to perform a blocking check to ensure that the interactor's observation schema is a valid subset of the recording asset's observation schema JSON description.
+	 * For example, this option can be left off to allow different observation encoding options in the interactor than that specified at recording time.
+	 * However, we recommend to keep this check unless you are confident of handling schema compatibility on your own.
+	 */
+	UPROPERTY(EditAnywhere, BlueprintReadWrite, Category = "LearningAgents")
+	bool bCheckObservationSchemaCompatibility = true;
+
+	/** 
+	 * Whether to perform a blocking check to ensure that the interactor's action schema is a valid subset of the recording asset's action schema JSON description.
+	 * For example, this option can be left off to allow different action encoding options in the interactor than that specified at recording time.
+	 * However, we recommend to keep this check unless you are confident of handling schema compatibility on your own.
+	 */
+	UPROPERTY(EditAnywhere, BlueprintReadWrite, Category = "LearningAgents")
+	bool bCheckActionSchemaCompatibility = true;
+};
+
+/** The configurable settings for the imitation learning training process. */
+USTRUCT(BlueprintType, Category = "LearningAgents")
+struct FLearningAgentsImitationTrainerTrainingSettings
+{
+	GENERATED_BODY()
+
+public:
+	UPROPERTY(EditAnywhere, BlueprintReadWrite, Category = "LearningAgents")
+	FLearningAgentsImitationTrainerHyperparameters Hyperparameters;
+
+	UPROPERTY(EditAnywhere, BlueprintReadWrite, Category = "LearningAgents")
+	FLearningAgentsImitationLearningLoggingSettings LoggingSettings;
+
+	UPROPERTY(EditAnywhere, BlueprintReadWrite, Category = "LearningAgents")
+	FLearningAgentsImitationLearningValidationSettings ValidationSettings;
+
+	/** The device to train on. */
+	UPROPERTY(EditAnywhere, BlueprintReadWrite, Category = "LearningAgents")
+	ELearningAgentsTrainingDevice Device = ELearningAgentsTrainingDevice::GPU;
 
 	/** Should evaluation run during the training process. Currently not used in Python */
 	UPROPERTY(VisibleDefaultsOnly, Category = "LearningAgents")
diff --git a/Engine/Plugins/Experimental/LearningAgents/Source/LearningAgentsTrainingEditor/Private/LearningAgentsImitationTrainerEditor.cpp b/Engine/Plugins/Experimental/LearningAgents/Source/LearningAgentsTrainingEditor/Private/LearningAgentsImitationTrainerEditor.cpp
--- a/Engine/Plugins/Experimental/LearningAgents/Source/LearningAgentsTrainingEditor/Private/LearningAgentsImitationTrainerEditor.cpp
+++ b/Engine/Plugins/Experimental/LearningAgents/Source/LearningAgentsTrainingEditor/Private/LearningAgentsImitationTrainerEditor.cpp
@@ -7,10 +7,7 @@
 
 #include UE_INLINE_GENERATED_CPP_BY_NAME(LearningAgentsImitationTrainerEditor)
 
-ALearningAgentsImitationTrainerEditor::ALearningAgentsImitationTrainerEditor()
-{
-	LearningAgentsManager = CreateDefaultSubobject<ULearningAgentsManager>(TEXT("LearningAgentsManager"));
-}
+ALearningAgentsImitationTrainerEditor::ALearningAgentsImitationTrainerEditor() {}
 
 void ALearningAgentsImitationTrainerEditor::StartTraining()
 {
@@ -62,13 +59,6 @@ void ALearningAgentsImitationTrainerEditor::StopTraining()
 	LearningAgentsImitationTrainer->EndTraining();
 }
 
-FLearningAgentsCommunicator ALearningAgentsImitationTrainerEditor::MakeFileCommunicator(FDirectoryPath EditorIntermediateRelativePath)
-{
-	FLearningAgentsCommunicator Communicator;
-	Communicator.Trainer = MakeShared<UE::Learning::FFileTrainer>(UE::Learning::Trainer::GetIntermediatePath(EditorIntermediateRelativePath.Path));
-	return Communicator;
-}
-
 bool ALearningAgentsImitationTrainerEditor::IsTraining() const
 {
 	if (LearningAgentsImitationTrainer)
diff --git a/Engine/Plugins/Experimental/LearningAgents/Source/LearningAgentsTrainingEditor/Private/LearningAgentsTrainingEditor.cpp b/Engine/Plugins/Experimental/LearningAgents/Source/LearningAgentsTrainingEditor/Private/LearningAgentsTrainingEditor.cpp
--- a/Engine/Plugins/Experimental/LearningAgents/Source/LearningAgentsTrainingEditor/Private/LearningAgentsTrainingEditor.cpp
+++ b/Engine/Plugins/Experimental/LearningAgents/Source/LearningAgentsTrainingEditor/Private/LearningAgentsTrainingEditor.cpp
@@ -1,8 +1,9 @@
 // Copyright Epic Games, Inc. All Rights Reserved.
 
 #include "LearningAgentsTrainingEditor.h"
+#include "LearningAgentsTrainerEditorDetails.h"
 #include "LearningAgentsImitationTrainerEditor.h"
-#include "LearningAgentsImitationTrainerEditorDetails.h"
+#include "LearningAgentsFlowMatchingTrainerEditor.h"
 #include "PropertyEditorModule.h"
 #include "PropertyEditorDelegates.h"
 #include "Misc/CoreDelegates.h"
@@ -12,7 +13,11 @@ void FLearningAgentsTrainingEditorModule::StartupModule()
 	FPropertyEditorModule& PropertyModule = FModuleManager::LoadModuleChecked<FPropertyEditorModule>("PropertyEditor");
 	PropertyModule.RegisterCustomClassLayout(
 		ALearningAgentsImitationTrainerEditor::StaticClass()->GetFName(),
-		FOnGetDetailCustomizationInstance::CreateStatic(&FLearningAgentsImitationTrainerEditorDetails::MakeInstance)
+		FOnGetDetailCustomizationInstance::CreateStatic(&FLearningAgentsTrainerEditorDetails::MakeInstance)
+	);
+	PropertyModule.RegisterCustomClassLayout(
+		ALearningAgentsFlowMatchingTrainerEditor::StaticClass()->GetFName(),
+		FOnGetDetailCustomizationInstance::CreateStatic(&FLearningAgentsTrainerEditorDetails::MakeInstance)
 	);
 }
 
@@ -22,6 +27,7 @@ void FLearningAgentsTrainingEditorModule::ShutdownModule()
 	{
 		FPropertyEditorModule& PropertyModule = FModuleManager::GetModuleChecked<FPropertyEditorModule>("PropertyEditor");
 		PropertyModule.UnregisterCustomClassLayout(ALearningAgentsImitationTrainerEditor::StaticClass()->GetFName());
+		PropertyModule.UnregisterCustomClassLayout(ALearningAgentsFlowMatchingTrainerEditor::StaticClass()->GetFName());
 	}
 }
 	
diff --git a/Engine/Plugins/Experimental/LearningAgents/Source/LearningAgentsTrainingEditor/Public/LearningAgentsImitationTrainerEditor.h b/Engine/Plugins/Experimental/LearningAgents/Source/LearningAgentsTrainingEditor/Public/LearningAgentsImitationTrainerEditor.h
--- a/Engine/Plugins/Experimental/LearningAgents/Source/LearningAgentsTrainingEditor/Public/LearningAgentsImitationTrainerEditor.h
+++ b/Engine/Plugins/Experimental/LearningAgents/Source/LearningAgentsTrainingEditor/Public/LearningAgentsImitationTrainerEditor.h
@@ -6,6 +6,7 @@
 #include "GameFramework/Actor.h"
 #include "LearningAgentsManager.h"
 #include "LearningAgentsCommunicator.h"
+#include "LearningAgentsTrainerEditorBase.h"
 #include "LearningAgentsImitationTrainer.h"
 #include "LearningAgentsImitationTrainerEditor.generated.h"
 
@@ -13,43 +14,24 @@
 
 /** Editor callable imitation learning trainer. */
 UCLASS(BlueprintType, Blueprintable)
-class ALearningAgentsImitationTrainerEditor : public AActor
+class ALearningAgentsImitationTrainerEditor : public ALearningAgentsTrainerEditorBase
 {
 	GENERATED_BODY()
 
 public:
 	UE_API ALearningAgentsImitationTrainerEditor();
 
-	/** Setup the imitation trainer with necessary components in blueprints. */
-	UFUNCTION(BlueprintImplementableEvent, Category = "LearningAgents")
-	UE_API void SetupTraining();
-
 	/** Start training. */
-	UFUNCTION(BlueprintCallable, Category = "LearningAgents")
-	UE_API void StartTraining();
+	UE_API virtual void StartTraining() override;
 
 	/** Stop training. */
-	UFUNCTION(BlueprintCallable, Category = "LearningAgents")
-	UE_API void StopTraining();
-
-	/** Make File communicator to file training materials for external training. */
-	UFUNCTION(BlueprintCallable, Category = "LearningAgents")
-	UE_API FLearningAgentsCommunicator MakeFileCommunicator(FDirectoryPath EditorIntermediateRelativePath);
-
-	UE_API bool IsTraining() const;
+	UE_API virtual void StopTraining() override;
 
-	UPROPERTY(EditDefaultsOnly, Category = "LearningAgents")
-	float TrainingTickInterval = 0.0f;
-	
-	UPROPERTY(VisibleAnywhere, BlueprintReadWrite, Category = "LearningAgents")
-    TObjectPtr<ULearningAgentsManager> LearningAgentsManager;
+	UE_API virtual bool IsTraining() const override;
 
 	UPROPERTY(BlueprintReadWrite, Category = "LearningAgents")
 	TObjectPtr<ULearningAgentsImitationTrainer> LearningAgentsImitationTrainer;
 
-	UPROPERTY(EditDefaultsOnly, BlueprintReadWrite, Category = "LearningAgents")
-	TObjectPtr<ULearningAgentsRecording> Recording;
-
 	UPROPERTY(EditDefaultsOnly, BlueprintReadWrite, Category = "LearningAgents")
 	FLearningAgentsImitationTrainerSettings ImitationTrainerSettings;
 
diff --git a/Engine/Plugins/Experimental/LearningCore/Source/Learning/Private/LearningAction.cpp b/Engine/Plugins/Experimental/LearningCore/Source/Learning/Private/LearningAction.cpp
--- a/Engine/Plugins/Experimental/LearningCore/Source/Learning/Private/LearningAction.cpp
+++ b/Engine/Plugins/Experimental/LearningCore/Source/Learning/Private/LearningAction.cpp
@@ -161,6 +161,12 @@ namespace UE::Learning::Action
 		{
 			return FMath::Loge(FMath::Max(X / FMath::Max(1.0f - X, FLT_MIN), FLT_MIN));
 		}
+
+		struct FContinuousDataDeprecated
+		{
+			int32 Num = INDEX_NONE;
+			float Scale = 0.0f;
+		};
 	}
 
 	FSchemaElement FSchema::CreateNull(const FName Tag)
@@ -179,11 +185,43 @@ namespace UE::Learning::Action
 	FSchemaElement FSchema::CreateContinuous(const FSchemaContinuousParameters Parameters, const FName Tag)
 	{
 		check(Parameters.Num >= 0);
-		check(Parameters.Scale >= 0.0f);
 
 		FContinuousData ElementData;
 		ElementData.Num = Parameters.Num;
-		ElementData.Scale = Parameters.Scale;
+		ElementData.NormalizationOption = Parameters.NormalizationOption;
+
+		if (Parameters.NormalizationOption == ENormalization::Manual)
+		{
+			check(Parameters.Scales.Num() == Parameters.Num or Parameters.Scales.IsEmpty());
+			check(Parameters.Offsets.Num() == Parameters.Num or Parameters.Offsets.IsEmpty());
+
+			ElementData.NormalizationIndexOffset = NormalizationScales.Num();
+
+			// Use provided arrays or identity (std=1, mean=0)
+			if (Parameters.Scales.Num() == Parameters.Num)
+			{
+				NormalizationScales.Append(Parameters.Scales.GetData(), Parameters.Scales.Num());
+			}
+			else
+			{
+				for (int32 i = 0; i < Parameters.Num; i++)
+				{
+					NormalizationScales.Add(1.0f);
+				}
+			}
+
+			if (Parameters.Offsets.Num() == Parameters.Num)
+			{
+				NormalizationOffsets.Append(Parameters.Offsets.GetData(), Parameters.Offsets.Num());
+			}
+			else
+			{
+				for (int32 i = 0; i < Parameters.Num; i++)
+				{
+					NormalizationOffsets.Add(0.0f);
+				}
+			}
+		}
 
 		const int32 Index = Types.Add(EType::Continuous);
 		Tags.Add(Tag);
@@ -467,7 +505,14 @@ namespace UE::Learning::Action
 
 		FSchemaContinuousParameters Parameters;
 		Parameters.Num = ElementData.Num;
-		Parameters.Scale = ElementData.Scale;
+		Parameters.NormalizationOption = ElementData.NormalizationOption;
+
+		if (Parameters.NormalizationOption == ENormalization::Manual)
+		{
+			Parameters.Scales = TArrayView<const float>(NormalizationScales.GetData() + ElementData.NormalizationIndexOffset, ElementData.Num);
+			Parameters.Offsets = TArrayView<const float>(NormalizationOffsets.GetData() + ElementData.NormalizationIndexOffset, ElementData.Num);
+		}
+				
 		return Parameters;
 	}
 
@@ -608,6 +653,8 @@ namespace UE::Learning::Action
 		SubElementNames.Empty();
 		SubElementObjects.Empty();
 		PriorProbabilities.Empty();
+		NormalizationScales.Empty();
+		NormalizationOffsets.Empty();
 
 		Generation++;
 	}
@@ -641,6 +688,8 @@ namespace UE::Learning::Action
 		SubElementNames.Reset();
 		SubElementObjects.Reset();
 		PriorProbabilities.Reset();
+		NormalizationScales.Reset();
+		NormalizationOffsets.Reset();
 
 		Generation++;
 	}
@@ -648,7 +697,9 @@ namespace UE::Learning::Action
 	bool FSchema::Serialize(FArchive& Ar)
 	{
 		static constexpr int32 Magic = 0x25344a5c;
-		static constexpr int32 Version = 0;
+		static constexpr int32 Version = 1;
+
+		bool bSerializeDeprecatedContinuousData = false;
 
 		if (Ar.IsLoading())
 		{
@@ -658,7 +709,9 @@ namespace UE::Learning::Action
 			Ar << LoadMagic;
 			if (!ensure(LoadMagic == Magic)) { Ar.Seek(Offset); Empty(); return false; }
 			Ar << LoadVersion;
-			if (!ensure(LoadVersion == Version)) { Ar.Seek(Offset); Empty(); return false; }
+
+			if (LoadVersion == 0) { bSerializeDeprecatedContinuousData = true; }
+			else if (!ensure(LoadVersion == Version)) { Ar.Seek(Offset); Empty(); return false; }
 		}
 		else
 		{
@@ -678,7 +731,33 @@ namespace UE::Learning::Action
 		Ar << ActionModifierVectorSizes;
 		Ar << TypeDataIndices;
 
-		Private::ArraySerializeRaw(Ar, ContinuousData);
+		if (bSerializeDeprecatedContinuousData)
+		{
+			TArray<Private::FContinuousDataDeprecated> DeprecatedContinuousData;
+			Private::ArraySerializeRaw(Ar, DeprecatedContinuousData);
+						
+			for (const Private::FContinuousDataDeprecated& DeprecatedData : DeprecatedContinuousData)
+			{
+				FContinuousData CurrentData;
+
+				CurrentData.Num = DeprecatedData.Num;
+				CurrentData.NormalizationOption = ENormalization::Manual;
+				CurrentData.NormalizationIndexOffset = NormalizationScales.Num();
+
+				for (int32 i = 0; i < DeprecatedData.Num; ++i)
+				{
+					NormalizationScales.Add(DeprecatedData.Scale);
+					NormalizationOffsets.Add(0.0f);
+				}
+
+				ContinuousData.Add(CurrentData);
+			}
+		}
+		else
+		{
+			Private::ArraySerializeRaw(Ar, ContinuousData);
+		}
+
 		Private::ArraySerializeRaw(Ar, DiscreteExclusiveData);
 		Private::ArraySerializeRaw(Ar, DiscreteInclusiveData);
 		Private::ArraySerializeRaw(Ar, NamedDiscreteExclusiveData);
@@ -694,6 +773,12 @@ namespace UE::Learning::Action
 
 		Ar << PriorProbabilities;
 
+		if (!bSerializeDeprecatedContinuousData)
+		{
+			Private::ArraySerializeRaw(Ar, NormalizationScales);
+			Private::ArraySerializeRaw(Ar, NormalizationOffsets);
+		}
+
 		return true;
 	}
 
@@ -2607,6 +2692,12 @@ namespace UE::Learning::Action
 			break;
 		}
 
+		default:
+		{
+			checkNoEntry();
+			return;
+		}
+
 		}
 	}
 		
@@ -2652,11 +2743,17 @@ namespace UE::Learning::Action
 			// Copy in and scale the values from the action object
 
 			const int32 ValueNum = SchemaParameters.Num;
-			const float ValueScale = FMath::Max(SchemaParameters.Scale, UE_SMALL_NUMBER);
 
 			for (int32 ValueIdx = 0; ValueIdx < ValueNum; ValueIdx++)
 			{
-				OutActionVector[ValueIdx] = ActionValues[ValueIdx] / ValueScale;
+				float Value = ActionValues[ValueIdx];				
+				if (SchemaParameters.NormalizationOption == ENormalization::Manual)
+				{
+					Value -= SchemaParameters.Offsets[ValueIdx];
+					const float StdDev = FMath::Max(SchemaParameters.Scales[ValueIdx], UE_SMALL_NUMBER);
+					Value /= StdDev;
+				}
+				OutActionVector[ValueIdx] = Value;
 			}
 
 			return;
@@ -2914,15 +3011,24 @@ namespace UE::Learning::Action
 			check(ActionVectorSize == SchemaParameters.Num);
 
 			const int32 ValueNum = SchemaParameters.Num;
-			const float ValueScale = FMath::Max(SchemaParameters.Scale, UE_SMALL_NUMBER);
 
 			TLearningArray<1, float, TInlineAllocator<32>> ActionValues;
 			ActionValues.SetNumUninitialized({ ValueNum });
+
 			for (int32 ValueIdx = 0; ValueIdx < ValueNum; ValueIdx++)
 			{
-				ActionValues[ValueIdx] = ValueScale * ActionVector[ValueIdx];
+				float Value = ActionVector[ValueIdx];
+				
+				if (SchemaParameters.NormalizationOption == ENormalization::Manual)
+				{
+					const float StdDev = FMath::Max(SchemaParameters.Scales[ValueIdx], UE_SMALL_NUMBER);
+					Value *= StdDev;
+					Value += SchemaParameters.Offsets[ValueIdx];
+				}
+				ActionValues[ValueIdx] = Value;
 			}
 
+
 			OutObjectElement = OutObject.CreateContinuous({ MakeArrayView(ActionValues.GetData(), ActionValues.Num()) }, SchemaElementTag);
 			return;
 		}
diff --git a/Engine/Plugins/Experimental/LearningCore/Source/Learning/Private/LearningObservation.cpp b/Engine/Plugins/Experimental/LearningCore/Source/Learning/Private/LearningObservation.cpp
--- a/Engine/Plugins/Experimental/LearningCore/Source/Learning/Private/LearningObservation.cpp
+++ b/Engine/Plugins/Experimental/LearningCore/Source/Learning/Private/LearningObservation.cpp
@@ -91,6 +91,12 @@ namespace UE::Learning::Observation
 			}
 			return true;
 		}
+
+		struct FContinuousDataDeprecated
+		{
+			int32 Num = INDEX_NONE;
+			float Scale = 0.0f;
+		};
 	}
 
 	FSchemaElement FSchema::CreateNull(const FName Tag)
@@ -108,11 +114,43 @@ namespace UE::Learning::Observation
 	FSchemaElement FSchema::CreateContinuous(const FSchemaContinuousParameters Parameters, const FName Tag)
 	{
 		check(Parameters.Num >= 0);
-		check(Parameters.Scale >= 0.0f);
 
 		FContinuousData ElementData;
 		ElementData.Num = Parameters.Num;
-		ElementData.Scale = Parameters.Scale;
+		ElementData.NormalizationOption = Parameters.NormalizationOption;
+
+		if (Parameters.NormalizationOption == ENormalization::Manual)
+		{
+			check(Parameters.Scales.Num() == Parameters.Num or Parameters.Scales.IsEmpty());
+			check(Parameters.Offsets.Num() == Parameters.Num or Parameters.Offsets.IsEmpty());
+
+			ElementData.NormalizationIndexOffset = NormalizationScales.Num();
+
+			// Use provided arrays or identity (std=1, mean=0)
+			if (Parameters.Scales.Num() == Parameters.Num)
+			{
+				NormalizationScales.Append(Parameters.Scales.GetData(), Parameters.Scales.Num());
+			}
+			else
+			{
+				for (int32 i = 0; i < Parameters.Num; i++)
+				{
+					NormalizationScales.Add(1.0f);
+				}
+			}
+
+			if (Parameters.Offsets.Num() == Parameters.Num)
+			{
+				NormalizationOffsets.Append(Parameters.Offsets.GetData(), Parameters.Offsets.Num());
+			}
+			else
+			{
+				for (int32 i = 0; i < Parameters.Num; i++)
+				{
+					NormalizationOffsets.Add(0.0f);
+				}
+			}
+		}
 
 		const int32 Index = Types.Add(EType::Continuous);
 		Tags.Add(Tag);
@@ -439,7 +477,14 @@ namespace UE::Learning::Observation
 
 		FSchemaContinuousParameters Parameters;
 		Parameters.Num = ElementData.Num;
-		Parameters.Scale = ElementData.Scale;
+		Parameters.NormalizationOption = ElementData.NormalizationOption;
+		
+		if (Parameters.NormalizationOption == ENormalization::Manual)
+		{
+			Parameters.Scales = TArrayView<const float>(NormalizationScales.GetData() + ElementData.NormalizationIndexOffset, ElementData.Num);
+			Parameters.Offsets = TArrayView<const float>(NormalizationOffsets.GetData() + ElementData.NormalizationIndexOffset, ElementData.Num);
+		}
+
 		return Parameters;
 	}
 
@@ -630,6 +675,9 @@ namespace UE::Learning::Observation
 		SubElementNames.Empty();
 		SubElementObjects.Empty();
 
+		NormalizationScales.Empty();
+		NormalizationOffsets.Empty();
+
 		Generation++;
 	}
 
@@ -664,13 +712,18 @@ namespace UE::Learning::Observation
 		SubElementNames.Reset();
 		SubElementObjects.Reset();
 
+		NormalizationScales.Reset();
+		NormalizationOffsets.Reset();
+
 		Generation++;
 	}
 
 	bool FSchema::Serialize(FArchive& Ar)
 	{
 		static constexpr int32 Magic = 0x7db41964;
-		static constexpr int32 Version = 0;
+		static constexpr int32 Version = 1;
+
+		bool bSerializeDeprecatedContinuousData = false;
 
 		if (Ar.IsLoading())
 		{
@@ -680,7 +733,9 @@ namespace UE::Learning::Observation
 			Ar << LoadMagic;
 			if (!ensure(LoadMagic == Magic)) { Ar.Seek(Offset); Empty(); return false; }
 			Ar << LoadVersion;
-			if (!ensure(LoadVersion == Version)) { Ar.Seek(Offset); Empty(); return false; }
+
+			if (LoadVersion == 0) { bSerializeDeprecatedContinuousData = true; }
+			else if (!ensure(LoadVersion == Version)) { Ar.Seek(Offset); Empty(); return false; }
 		}
 		else
 		{
@@ -699,7 +754,33 @@ namespace UE::Learning::Observation
 		Ar << ObservationDistributionVectorSizes;
 		Ar << TypeDataIndices;
 
-		Private::ArraySerializeRaw(Ar, ContinuousData);
+		if (bSerializeDeprecatedContinuousData)
+		{
+			TArray<Private::FContinuousDataDeprecated> DeprecatedContinuousData;
+			Private::ArraySerializeRaw(Ar, DeprecatedContinuousData);
+
+			for (const Private::FContinuousDataDeprecated& DeprecatedData : DeprecatedContinuousData)
+			{
+				FContinuousData CurrentData;
+
+				CurrentData.Num = DeprecatedData.Num;
+				CurrentData.NormalizationOption = ENormalization::Manual;
+				CurrentData.NormalizationIndexOffset = NormalizationScales.Num();
+
+				for (int32 i = 0; i < DeprecatedData.Num; ++i)
+				{
+					NormalizationScales.Add(DeprecatedData.Scale);
+					NormalizationOffsets.Add(0.0f);
+				}
+
+				ContinuousData.Add(CurrentData);
+			}
+		}
+		else
+		{
+			Private::ArraySerializeRaw(Ar, ContinuousData);
+		}
+
 		Private::ArraySerializeRaw(Ar, DiscreteExclusiveData);
 		Private::ArraySerializeRaw(Ar, DiscreteInclusiveData);
 		Private::ArraySerializeRaw(Ar, NamedDiscreteExclusiveData);
@@ -716,6 +797,12 @@ namespace UE::Learning::Observation
 		Ar << SubElementNames;
 		Ar << SubElementObjects;
 
+		if (!bSerializeDeprecatedContinuousData)
+		{
+			Private::ArraySerializeRaw(Ar, NormalizationScales);
+			Private::ArraySerializeRaw(Ar, NormalizationOffsets);
+		}
+
 		return true;
 	}
 
@@ -2018,8 +2105,9 @@ namespace UE::Learning::Observation
 			// Check the input sizes match
 
 			const FSchemaContinuousParameters SchemaParameters = Schema.GetContinuous(SchemaElement);
+			const FObjectContinuousParameters ObjectParameters = Object.GetContinuous(ObjectElement);
 
-			const TArrayView<const float> ObservationValues = Object.GetContinuous(ObjectElement).Values;
+			const TArrayView<const float> ObservationValues = ObjectParameters.Values;
 			check(Schema.GetObservationVectorSize(SchemaElement) == ObservationValues.Num());
 			check(Schema.GetObservationVectorSize(SchemaElement) == OutObservationVector.Num());
 			check(Schema.GetObservationVectorSize(SchemaElement) == SchemaParameters.Num);
@@ -2027,13 +2115,19 @@ namespace UE::Learning::Observation
 			// Copy in and scale the values from the observation object
 
 			const int32 ValueNum = SchemaParameters.Num;
-			const float ValueScale = FMath::Max(SchemaParameters.Scale, UE_SMALL_NUMBER);
-
 			for (int32 ValueIdx = 0; ValueIdx < ValueNum; ValueIdx++)
 			{
-				OutObservationVector[ValueIdx] = ObservationValues[ValueIdx] / ValueScale;
-			}
+				float Value = ObservationValues[ValueIdx];
 
+				if (SchemaParameters.NormalizationOption == ENormalization::Manual)
+				{
+					Value -= SchemaParameters.Offsets[ValueIdx];
+					const float StdDev = FMath::Max(SchemaParameters.Scales[ValueIdx], UE_SMALL_NUMBER);
+					Value /= StdDev;
+				}
+				
+				OutObservationVector[ValueIdx] = Value;
+			}
 			return;
 		}
 
@@ -2356,19 +2450,27 @@ namespace UE::Learning::Observation
 			check(ObservationVectorSize == SchemaParameters.Num);
 
 			const int32 ValueNum = SchemaParameters.Num;
-			const float ValueScale = FMath::Max(SchemaParameters.Scale, UE_SMALL_NUMBER);
 
 			TLearningArray<1, float, TInlineAllocator<32>> ObservationValues;
 			ObservationValues.SetNumUninitialized({ ValueNum });
+
 			for (int32 ValueIdx = 0; ValueIdx < ValueNum; ValueIdx++)
 			{
-				ObservationValues[ValueIdx] = ValueScale * ObservationVector[ValueIdx];
+				float Value = ObservationVector[ValueIdx];
+
+				if (SchemaParameters.NormalizationOption == ENormalization::Manual)
+				{
+					const float StdDev = FMath::Max(SchemaParameters.Scales[ValueIdx], UE_SMALL_NUMBER);
+					Value *= StdDev;
+					Value += SchemaParameters.Offsets[ValueIdx];
+				}
+				ObservationValues[ValueIdx] = Value;
 			}
 
 			OutObjectElement = OutObject.CreateContinuous({ MakeArrayView(ObservationValues.GetData(), ObservationValues.Num()) }, SchemaElementTag);
 			return;
 		}
-		
+
 		case EType::DiscreteExclusive:
 		{
 			check(ObservationVectorSize == Schema.GetDiscreteExclusive(SchemaElement).Num);
diff --git a/Engine/Plugins/Experimental/LearningCore/Source/Learning/Public/LearningAction.h b/Engine/Plugins/Experimental/LearningCore/Source/Learning/Public/LearningAction.h
--- a/Engine/Plugins/Experimental/LearningCore/Source/Learning/Public/LearningAction.h
+++ b/Engine/Plugins/Experimental/LearningCore/Source/Learning/Public/LearningAction.h
@@ -57,6 +57,24 @@ namespace UE::Learning::Action
 		Encoding			= 10,
 	};
 
+	enum class ENormalization : uint8
+	{
+		Manual = 0,
+		AutoShared = 1,
+		AutoPerDimension = 2,
+	};
+
+	static inline FString GetNormalizationOptionAsString(const ENormalization NormalizationOption)
+	{
+		switch (NormalizationOption)
+		{
+		case UE::Learning::Action::ENormalization::Manual: return TEXT("Manual");
+		case UE::Learning::Action::ENormalization::AutoShared: return TEXT("AutoShared");
+		case UE::Learning::Action::ENormalization::AutoPerDimension: return TEXT("AutoPerDimension");
+		default: checkNoEntry(); return TEXT("");
+		}
+	}
+
 	/**
 	 * Action Schema Element
 	 * 
@@ -70,8 +88,14 @@ namespace UE::Learning::Action
 		// Number of values in the continuous action
 		int32 Num = 0;
 
-		// Scale factor for the continuous action
-		float Scale = 1.0f;
+		// Per-element standard deviation values for normalization. If empty, identity scaling is applied. *Only used when normalization option is set to manual.
+		TArrayView<const float> Scales;
+
+		// Per-element mean values for normalization. If empty, zero bias is applied. *Only used when normalization option is set to manual.
+		TArrayView<const float> Offsets;
+
+		// The normalization option to use.
+		ENormalization NormalizationOption = ENormalization::Manual;
 	};
 
 	struct FSchemaDiscreteExclusiveParameters
@@ -264,7 +288,8 @@ namespace UE::Learning::Action
 		struct FContinuousData
 		{
 			int32 Num = INDEX_NONE;
-			float Scale = 0.0f;
+			int32 NormalizationIndexOffset = INDEX_NONE;
+			ENormalization NormalizationOption = ENormalization::Manual;
 		};
 
 		struct FDiscreteExclusiveData
@@ -356,6 +381,10 @@ namespace UE::Learning::Action
 
 		/** This is an array of all the prior probabilities, referenced by other elements. */
 		TArray<float> PriorProbabilities;
+
+		/** These are arrays for normalization values, referenced by continuous and flow matching elements. */
+		TArray<float> NormalizationScales;
+		TArray<float> NormalizationOffsets;
 	};
 
 	/**
@@ -435,6 +464,12 @@ namespace UE::Learning::Action
 		FObjectElement Element;
 	};
 
+	struct FObjectFlowMatchingParameters
+	{
+		// Flow matching float action values
+		TArrayView<const float> Values;
+	};
+
 	/**
 	 * Action Object
 	 *
@@ -455,6 +490,7 @@ namespace UE::Learning::Action
 		UE_API FObjectElement CreateOrInclusive(const FObjectOrInclusiveParameters Parameters, const FName Tag = NAME_None);
 		UE_API FObjectElement CreateArray(const FObjectArrayParameters Parameters, const FName Tag = NAME_None);
 		UE_API FObjectElement CreateEncoding(const FObjectEncodingParameters Parameters, const FName Tag = NAME_None);
+		UE_API FObjectElement CreateFlowMatching(const FObjectFlowMatchingParameters Parameters, const FName Tag = NAME_None);
 
 		UE_API FObjectContinuousParameters GetContinuous(const FObjectElement Element) const;
 		UE_API FObjectDiscreteExclusiveParameters GetDiscreteExclusive(const FObjectElement Element) const;
@@ -466,6 +502,7 @@ namespace UE::Learning::Action
 		UE_API FObjectOrInclusiveParameters GetOrInclusive(const FObjectElement Element) const;
 		UE_API FObjectArrayParameters GetArray(const FObjectElement Element) const;
 		UE_API FObjectEncodingParameters GetEncoding(const FObjectElement Element) const;
+		UE_API FObjectFlowMatchingParameters GetFlowMatching(const FObjectElement Element) const;
 
 		// Checks if the given element is valid
 		UE_API bool IsValid(const FObjectElement Element) const;
diff --git a/Engine/Plugins/Experimental/LearningCore/Source/Learning/Public/LearningObservation.h b/Engine/Plugins/Experimental/LearningCore/Source/Learning/Public/LearningObservation.h
--- a/Engine/Plugins/Experimental/LearningCore/Source/Learning/Public/LearningObservation.h
+++ b/Engine/Plugins/Experimental/LearningCore/Source/Learning/Public/LearningObservation.h
@@ -66,6 +66,24 @@ namespace UE::Learning::Observation
 		Conv2d = 13,
 	};
 
+	enum class ENormalization : uint8
+	{
+		Manual = 0,
+		AutoShared = 1,
+		AutoPerDimension = 2,
+	};
+
+	static inline FString GetNormalizationOptionAsString(const ENormalization NormalizationOption)
+	{
+		switch (NormalizationOption)
+		{
+		case UE::Learning::Observation::ENormalization::Manual: return TEXT("Manual");
+		case UE::Learning::Observation::ENormalization::AutoShared: return TEXT("AutoShared");
+		case UE::Learning::Observation::ENormalization::AutoPerDimension: return TEXT("AutoPerDimension");
+		default: checkNoEntry(); return TEXT("");
+		}
+	}
+
 	/**
 	 * Observation Schema Element
 	 *
@@ -79,8 +97,14 @@ namespace UE::Learning::Observation
 		// Number of values in the continuous observation
 		int32 Num = 0;
 
-		// Scale factor for the continuous observation
-		float Scale = 1.0f;
+		// Per-element standard deviation values for normalization. If empty, identity scaling is applied. *Only used when normalization option is set to manual.
+		TArrayView<const float> Scales;
+
+		// Per-element mean values for normalization. If empty, zero bias is applied. *Only used when normalization option is set to manual.
+		TArrayView<const float> Offsets;
+
+		// The normalization option to use
+		ENormalization NormalizationOption = ENormalization::Manual;
 	};
 
 	struct FSchemaDiscreteExclusiveParameters
@@ -343,7 +367,8 @@ namespace UE::Learning::Observation
 		struct FContinuousData
 		{
 			int32 Num = INDEX_NONE;
-			float Scale = 0.0f;
+			int32 NormalizationIndexOffset = INDEX_NONE;
+			ENormalization NormalizationOption = ENormalization::Manual;
 		};
 
 		struct FDiscreteExclusiveData
@@ -467,6 +492,10 @@ namespace UE::Learning::Observation
 		/** This is an array of all the SubElements and their names, referenced by other elements */
 		TArray<FName> SubElementNames;
 		TArray<FSchemaElement> SubElementObjects;
+
+		/** Normalization arrays for continuous observations, indexed by NormalizationOffset in FContinuousData */
+		TArray<float> NormalizationScales;
+		TArray<float> NormalizationOffsets;
 	};
 
 	/**
diff --git a/Engine/Plugins/Experimental/LearningCore/Source/LearningTraining/Private/LearningTrainer.cpp b/Engine/Plugins/Experimental/LearningCore/Source/LearningTraining/Private/LearningTrainer.cpp
--- a/Engine/Plugins/Experimental/LearningCore/Source/LearningTraining/Private/LearningTrainer.cpp
+++ b/Engine/Plugins/Experimental/LearningCore/Source/LearningTraining/Private/LearningTrainer.cpp
@@ -123,6 +123,7 @@ namespace UE::Learning::Trainer
 			const Observation::FSchemaContinuousParameters Parameters = ObservationSchema.GetContinuous(ObservationSchemaElement);
 			Object->SetStringField(TEXT("Type"), TEXT("Continuous"));
 			Object->SetNumberField(TEXT("Num"), Parameters.Num);
+			Object->SetStringField(TEXT("Normalization"), Observation::GetNormalizationOptionAsString(Parameters.NormalizationOption));
 			break;
 		}
 
@@ -328,6 +329,8 @@ namespace UE::Learning::Trainer
 
 			Object->SetStringField(TEXT("Type"), TEXT("Continuous"));
 			Object->SetNumberField(TEXT("Num"), Parameters.Num);
+			Object->SetStringField(TEXT("Normalization"), Action::GetNormalizationOptionAsString(Parameters.NormalizationOption));
+
 			break;
 		}
 
@@ -768,6 +771,13 @@ namespace UE::Learning::Trainer
 			// Do nothing since Null is zero-width
 			return;
 		}
+		if ((Object->GetStringField(TEXT("Type")) != TEXT("Null")) && ObservationSchema.GetType(ObservationSchemaElement) == Observation::EType::Null)
+		{
+			// Skipping element: schema has Null but source JSON has a non null type. This is an acceptable method to skip using certain actions for training
+			int32 VectorSize = Object->GetNumberField(TEXT("VectorSize"));
+			NextIndex += VectorSize;
+			return;
+		}
 		else if (Object->GetStringField(TEXT("Type")) == TEXT("Continuous"))
 		{
 			const Observation::FSchemaContinuousParameters Parameters = ObservationSchema.GetContinuous(ObservationSchemaElement);
@@ -1270,6 +1280,13 @@ namespace UE::Learning::Trainer
 			// Do nothing since Null is zero-width
 			return;
 		}
+		if ((Object->GetStringField(TEXT("Type")) != TEXT("Null")) && ActionSchema.GetType(ActionSchemaElement) == Action::EType::Null)
+		{
+			// Skipping element: schema has Null but source JSON has a non null type. This is an acceptable method to skip using certain actions for training
+			int32 VectorSize = Object->GetNumberField(TEXT("VectorSize"));
+			NextIndex += VectorSize;
+			return;
+		}
 		else if (Object->GetStringField(TEXT("Type")) == TEXT("Continuous"))
 		{
 			const Action::FSchemaContinuousParameters Parameters = ActionSchema.GetContinuous(ActionSchemaElement);
