[Backout] - CL51610429
#rnx
Original CL Desc
-----------------------------------------------------------------
[LearningAgents] Flow Matching and Normalizations

+ network embedded normalizations
+ flow matching training objective

- update anim control as a result of new normalization code
- update LACombat as a result of new normalization code

diff --git a/Engine/Plugins/Experimental/Animation/AnimGen/Source/AnimGen/Private/AnimGenControl.cpp b/Engine/Plugins/Experimental/Animation/AnimGen/Source/AnimGen/Private/AnimGenControl.cpp
--- a/Engine/Plugins/Experimental/Animation/AnimGen/Source/AnimGen/Private/AnimGenControl.cpp
+++ b/Engine/Plugins/Experimental/Animation/AnimGen/Source/AnimGen/Private/AnimGenControl.cpp
@@ -735,13 +735,9 @@ FAnimGenControlSchemaElement UAnimGenControls::SpecifyContinuousControl(FAnimGen
 		UE_ANIMGEN_LOG_ONCE(LogAnimGen, Warning, "SpecifyContinuousControl: Specifying zero-sized Continuous Control.");
 	}
 
-	TArray<float, TInlineAllocator<128>> Scales, Offsets;
-	Scales.Init(Scale, Size);
-	Offsets.Init(0.0f, Size);
-
 	{
 		FWriteScopeLock SchemaLock(Schema.ObservationSchema->Lock);
-		return { Schema.ObservationSchema->Schema.CreateContinuous({ Size, Scales, Offsets, UE::Learning::Observation::ENormalization::Manual }, Tag) };
+		return { Schema.ObservationSchema->Schema.CreateContinuous({ Size, Scale }, Tag) };
 	}
 }
 
diff --git a/Engine/Plugins/Experimental/LearningAgents/Source/LearningAgents/Private/LearningAgentsActions.cpp b/Engine/Plugins/Experimental/LearningAgents/Source/LearningAgents/Private/LearningAgentsActions.cpp
--- a/Engine/Plugins/Experimental/LearningAgents/Source/LearningAgents/Private/LearningAgentsActions.cpp
+++ b/Engine/Plugins/Experimental/LearningAgents/Source/LearningAgents/Private/LearningAgentsActions.cpp
@@ -32,7 +32,7 @@ uint32 GetTypeHash(const FLearningAgentsActionSchemaElement& Element)
 	return (uint32)Element.SchemaElement.Index;
 }
 
-uint32 GetTypeHash(const FLearningAgentsActionObjectElement& Element)
+uint32 GetTypeHash(const FLearningAgentsActionObjectElement & Element)
 {
 	return (uint32)Element.ObjectElement.Index;
 }
@@ -387,10 +387,10 @@ namespace UE::Learning::Agents::Action::Private
 			const Learning::Action::FObjectEncodingParameters ObjectParameters = Object.GetEncoding(ObjectElement);
 
 			return ValidateActionObjectMatchesSchema(
-				Schema,
-				SchemaParameters.Element,
-				Object,
-				ObjectParameters.Element);
+					Schema,
+					SchemaParameters.Element,
+					Object,
+					ObjectParameters.Element);
 		}
 
 		default:
@@ -922,7 +922,6 @@ namespace UE::Learning::Agents::Action::Private
 
 			return;
 		}
-
 		}
 	}
 
@@ -998,17 +997,6 @@ namespace UE::Learning::Agents::Action::Private
 		default: checkNoEntry(); return Learning::Action::EEncodingActivationFunction::ReLU;
 		}
 	}
-
-	static inline Learning::Action::ENormalization GetNormalizationOption(const ELearningAgentsNormalization NormalizationOption)
-	{
-		switch (NormalizationOption)
-		{
-		case ELearningAgentsNormalization::Manual: return Learning::Action::ENormalization::Manual;
-		case ELearningAgentsNormalization::AutoShared: return Learning::Action::ENormalization::AutoShared;
-		case ELearningAgentsNormalization::AutoPerDimension: return Learning::Action::ENormalization::AutoPerDimension;
-		default: checkNoEntry(); return Learning::Action::ENormalization::Manual;
-		}
-	}
 }
 
 bool ULearningAgentsActions::ValidateActionObjectMatchesSchema(
@@ -1072,7 +1060,7 @@ FLearningAgentsActionSchemaElement ULearningAgentsActions::SpecifyNullAction(ULe
 	return { Schema->ActionSchema.CreateNull(Tag) };
 }
 
-FLearningAgentsActionSchemaElement ULearningAgentsActions::SpecifyContinuousAction(ULearningAgentsActionSchema* Schema, const int32 Size, const FLearningAgentsActionNormalizationSettings& NormalizationSettings, const bool bActive, const FName Tag)
+FLearningAgentsActionSchemaElement ULearningAgentsActions::SpecifyContinuousAction(ULearningAgentsActionSchema* Schema, const int32 Size, const float Scale, const bool bActive, const FName Tag)
 {
 	if (!bActive)
 	{
@@ -1084,19 +1072,25 @@ FLearningAgentsActionSchemaElement ULearningAgentsActions::SpecifyContinuousActi
 		UE_LOG(LogLearning, Error, TEXT("SpecifyContinuousAction: Schema is nullptr."));
 		return FLearningAgentsActionSchemaElement();
 	}
-
+	
 	if (Size < 0)
 	{
 		UE_LOG(LogLearning, Error, TEXT("SpecifyContinuousAction: Invalid Continuous Action Size '%i'."), Size);
 		return FLearningAgentsActionSchemaElement();
 	}
 
+	if (Scale < UE_SMALL_NUMBER)
+	{
+		UE_LOG(LogLearning, Error, TEXT("SpecifyContinuousAction: Invalid Scale for Continuous Action '%f', must be greater than '%f'."), Scale, UE_SMALL_NUMBER);
+		return FLearningAgentsActionSchemaElement();
+	}
+
 	if (Size == 0)
 	{
 		UE_LOG(LogLearning, Warning, TEXT("SpecifyContinuousAction: Specifying zero-sized Continuous Action."));
 	}
 
-	return { Schema->ActionSchema.CreateContinuous({ Size, NormalizationSettings.Scales, NormalizationSettings.Offsets, UE::Learning::Agents::Action::Private::GetNormalizationOption(NormalizationSettings.NormalizationOption)}, Tag) };
+	return { Schema->ActionSchema.CreateContinuous({ Size, Scale }, Tag) };
 }
 
 FLearningAgentsActionSchemaElement ULearningAgentsActions::SpecifyExclusiveDiscreteAction(ULearningAgentsActionSchema* Schema, const int32 Size, const TArray<float>& PriorProbabilities, const bool bActive, const FName Tag)
@@ -1147,7 +1141,7 @@ FLearningAgentsActionSchemaElement ULearningAgentsActions::SpecifyExclusiveDiscr
 	}
 	UE::Learning::Agents::Action::Private::NormalizeProbabilitiesExclusive(NormalizedPriorProbabilities);
 
-	return { Schema->ActionSchema.CreateDiscreteExclusive({ Size, MakeArrayView(NormalizedPriorProbabilities) }, Tag) };
+	return { Schema->ActionSchema.CreateDiscreteExclusive({ Size, MakeArrayView(NormalizedPriorProbabilities) }, Tag)};
 }
 
 FLearningAgentsActionSchemaElement ULearningAgentsActions::SpecifyNamedExclusiveDiscreteAction(ULearningAgentsActionSchema* Schema, const TArray<FName>& Names, const TMap<FName, float>& PriorProbabilities, const bool bActive, const FName Tag)
@@ -1417,9 +1411,9 @@ FLearningAgentsActionSchemaElement ULearningAgentsActions::SpecifyStructAction(U
 	// Sort Elements According to FName
 
 	SubElementIndices.Sort([SubElementNames](const int32 Lhs, const int32 Rhs)
-		{
-			return SubElementNames[Lhs].ToString().ToLower() < SubElementNames[Rhs].ToString().ToLower();
-		});
+	{
+		return SubElementNames[Lhs].ToString().ToLower() < SubElementNames[Rhs].ToString().ToLower();
+	});
 
 	TArray<FName, TInlineAllocator<16>> SortedSubElementNames;
 	TArray<FLearningAgentsActionSchemaElement, TInlineAllocator<16>> SortedSubElements;
@@ -1528,9 +1522,9 @@ FLearningAgentsActionSchemaElement ULearningAgentsActions::SpecifyExclusiveUnion
 	// Sort Elements According to FName
 
 	SubElementIndices.Sort([SubElementNames](const int32 Lhs, const int32 Rhs)
-		{
-			return SubElementNames[Lhs].ToString().ToLower() < SubElementNames[Rhs].ToString().ToLower();
-		});
+	{
+		return SubElementNames[Lhs].ToString().ToLower() < SubElementNames[Rhs].ToString().ToLower();
+	});
 
 	TArray<FName, TInlineAllocator<16>> SortedSubElementNames;
 	TArray<FLearningAgentsActionSchemaElement, TInlineAllocator<16>> SortedSubElements;
@@ -1616,7 +1610,7 @@ FLearningAgentsActionSchemaElement ULearningAgentsActions::SpecifyExclusiveUnion
 	}
 	UE::Learning::Agents::Action::Private::NormalizeProbabilitiesExclusive(NormalizedPriorProbabilities);
 
-	return { Schema->ActionSchema.CreateOrExclusive({ ElementNames, SubElements, MakeArrayView(NormalizedPriorProbabilities) }, Tag) };
+	return { Schema->ActionSchema.CreateOrExclusive({ ElementNames, SubElements, MakeArrayView(NormalizedPriorProbabilities) }, Tag)};
 }
 
 FLearningAgentsActionSchemaElement ULearningAgentsActions::SpecifyInclusiveUnionAction(ULearningAgentsActionSchema* Schema, const TMap<FName, FLearningAgentsActionSchemaElement>& Elements, const TMap<FName, float>& PriorProbabilities, const bool bActive, const FName Tag)
@@ -1656,9 +1650,9 @@ FLearningAgentsActionSchemaElement ULearningAgentsActions::SpecifyInclusiveUnion
 	// Sort Elements According to FName
 
 	SubElementIndices.Sort([SubElementNames](const int32 Lhs, const int32 Rhs)
-		{
-			return SubElementNames[Lhs].ToString().ToLower() < SubElementNames[Rhs].ToString().ToLower();
-		});
+	{
+		return SubElementNames[Lhs].ToString().ToLower() < SubElementNames[Rhs].ToString().ToLower();
+	});
 
 	TArray<FName, TInlineAllocator<16>> SortedSubElementNames;
 	TArray<FLearningAgentsActionSchemaElement, TInlineAllocator<16>> SortedSubElements;
@@ -1744,7 +1738,7 @@ FLearningAgentsActionSchemaElement ULearningAgentsActions::SpecifyInclusiveUnion
 	}
 	UE::Learning::Agents::Action::Private::NormalizeProbabilitiesInclusive(NormalizedPriorProbabilities);
 
-	return { Schema->ActionSchema.CreateOrInclusive({ ElementNames, SubElements, MakeArrayView(NormalizedPriorProbabilities) }, Tag) };
+	return { Schema->ActionSchema.CreateOrInclusive({ ElementNames, SubElements, MakeArrayView(NormalizedPriorProbabilities) }, Tag)};
 }
 
 FLearningAgentsActionSchemaElement ULearningAgentsActions::SpecifyStaticArrayAction(ULearningAgentsActionSchema* Schema, const FLearningAgentsActionSchemaElement Element, const int32 Num, const bool bActive, const FName Tag)
@@ -1971,7 +1965,7 @@ FLearningAgentsActionSchemaElement ULearningAgentsActions::SpecifyEncodingAction
 		return FLearningAgentsActionSchemaElement();
 	}
 
-	return { Schema->ActionSchema.CreateEncoding({ Element.SchemaElement, EncodingSize, HiddenLayerNum, UE::Learning::Agents::Action::Private::GetEncodingActivationFunction(ActivationFunction) }, Tag) };
+	return { Schema->ActionSchema.CreateEncoding({ Element.SchemaElement, EncodingSize, HiddenLayerNum, UE::Learning::Agents::Action::Private::GetEncodingActivationFunction(ActivationFunction) }, Tag)};
 }
 
 FLearningAgentsActionSchemaElement ULearningAgentsActions::SpecifyBoolAction(ULearningAgentsActionSchema* Schema, const float PriorProbability, const bool bActive, const FName Tag)
@@ -1984,71 +1978,51 @@ FLearningAgentsActionSchemaElement ULearningAgentsActions::SpecifyBoolAction(ULe
 	return SpecifyExclusiveDiscreteActionFromArrayView(Schema, 2, { 1.0f - PriorProbability, PriorProbability }, bActive, Tag);
 }
 
-FLearningAgentsActionSchemaElement ULearningAgentsActions::SpecifyFloatAction(ULearningAgentsActionSchema* Schema, const FLearningAgentsFloatActionNormalizationSettings& NormalizationSettings, const bool bActive, const FName Tag)
+FLearningAgentsActionSchemaElement ULearningAgentsActions::SpecifyFloatAction(ULearningAgentsActionSchema* Schema, const float FloatScale, const bool bActive, const FName Tag)
 {
 	if (!bActive)
 	{
 		return SpecifyNullAction(Schema, Tag);
 	}
 
-	FLearningAgentsActionNormalizationSettings Settings;
-	Settings.Scales = { NormalizationSettings.Scale };
-	Settings.Offsets = { NormalizationSettings.Offset };
-	Settings.NormalizationOption = { NormalizationSettings.NormalizationOption };
-
-	return SpecifyContinuousAction(Schema, 1, Settings, bActive, Tag);
+	return SpecifyContinuousAction(Schema, 1, FloatScale, bActive, Tag);
 }
 
-FLearningAgentsActionSchemaElement ULearningAgentsActions::SpecifyLocationAction(ULearningAgentsActionSchema* Schema, const FLearningAgentsVectorActionNormalizationSettings& NormalizationSettings, const bool bActive, const FName Tag)
+FLearningAgentsActionSchemaElement ULearningAgentsActions::SpecifyLocationAction(ULearningAgentsActionSchema* Schema, const float LocationScale, const bool bActive, const FName Tag)
 {
 	if (!bActive)
 	{
 		return SpecifyNullAction(Schema, Tag);
 	}
 
-	FLearningAgentsActionNormalizationSettings Settings;
-	Settings.Scales = { (float)NormalizationSettings.Scales.X, (float)NormalizationSettings.Scales.Y, (float)NormalizationSettings.Scales.Z };
-	Settings.Offsets = { (float)NormalizationSettings.Offsets.X, (float)NormalizationSettings.Offsets.Y, (float)NormalizationSettings.Offsets.Z };
-	Settings.NormalizationOption = { NormalizationSettings.NormalizationOption };
-
-	return SpecifyContinuousAction(Schema, 3, Settings, bActive, Tag);
+	return SpecifyContinuousAction(Schema, 3, LocationScale, bActive, Tag);
 }
 
-FLearningAgentsActionSchemaElement ULearningAgentsActions::SpecifyRotationAction(ULearningAgentsActionSchema* Schema, const FLearningAgentsVectorActionNormalizationSettings& NormalizationSettings, const bool bActive, const FName Tag)
+FLearningAgentsActionSchemaElement ULearningAgentsActions::SpecifyRotationAction(ULearningAgentsActionSchema* Schema, const float RotationScale, const bool bActive, const FName Tag)
 {
 	if (!bActive)
 	{
 		return SpecifyNullAction(Schema, Tag);
 	}
 
-	FLearningAgentsActionNormalizationSettings Settings;
-	Settings.Scales = { (float)NormalizationSettings.Scales.X, (float)NormalizationSettings.Scales.Y, (float)NormalizationSettings.Scales.Z };
-	Settings.Offsets = { (float)NormalizationSettings.Offsets.X, (float)NormalizationSettings.Offsets.Y, (float)NormalizationSettings.Offsets.Z };
-	Settings.NormalizationOption = { NormalizationSettings.NormalizationOption };
-
-	return SpecifyContinuousAction(Schema, 3, Settings, bActive, Tag);
+	return SpecifyContinuousAction(Schema, 3, RotationScale, bActive, Tag);
 }
 
-FLearningAgentsActionSchemaElement ULearningAgentsActions::SpecifyScaleAction(ULearningAgentsActionSchema* Schema, const FLearningAgentsVectorActionNormalizationSettings& NormalizationSettings, const bool bActive, const FName Tag)
+FLearningAgentsActionSchemaElement ULearningAgentsActions::SpecifyScaleAction(ULearningAgentsActionSchema* Schema, const float ScaleScale, const bool bActive, const FName Tag)
 {
 	if (!bActive)
 	{
 		return SpecifyNullAction(Schema, Tag);
 	}
 
-	FLearningAgentsActionNormalizationSettings Settings;
-	Settings.Scales = { (float)NormalizationSettings.Scales.X, (float)NormalizationSettings.Scales.Y, (float)NormalizationSettings.Scales.Z };
-	Settings.Offsets = { (float)NormalizationSettings.Offsets.X, (float)NormalizationSettings.Offsets.Y, (float)NormalizationSettings.Offsets.Z };
-	Settings.NormalizationOption = NormalizationSettings.NormalizationOption;
-
-	return SpecifyContinuousAction(Schema, 3, Settings, bActive, Tag);
+	return SpecifyContinuousAction(Schema, 3, ScaleScale, bActive, Tag);
 }
 
 FLearningAgentsActionSchemaElement ULearningAgentsActions::SpecifyTransformAction(
-	ULearningAgentsActionSchema* Schema,
-	const FLearningAgentsVectorActionNormalizationSettings& LocationNormalization,
-	const FLearningAgentsVectorActionNormalizationSettings& RotationNormalization,
-	const FLearningAgentsVectorActionNormalizationSettings& ScaleNormalization,
+	ULearningAgentsActionSchema* Schema, 
+	const float LocationScale ,
+	const float RotationScale,
+	const float ScaleScale,
 	const bool bActive,
 	const FName Tag)
 {
@@ -2064,41 +2038,31 @@ FLearningAgentsActionSchemaElement ULearningAgentsActions::SpecifyTransformActio
 			TEXT("Scale")
 		},
 			{
-				SpecifyLocationAction(Schema, LocationNormalization),
-				SpecifyRotationAction(Schema, RotationNormalization),
-				SpecifyScaleAction(Schema, ScaleNormalization)
+				SpecifyLocationAction(Schema, LocationScale),
+				SpecifyRotationAction(Schema, RotationScale),
+				SpecifyScaleAction(Schema, ScaleScale)
 			},
 		bActive, Tag);
 }
 
-FLearningAgentsActionSchemaElement ULearningAgentsActions::SpecifyAngleAction(ULearningAgentsActionSchema* Schema, const FLearningAgentsFloatActionNormalizationSettings& NormalizationSettings, const bool bActive, const FName Tag)
+FLearningAgentsActionSchemaElement ULearningAgentsActions::SpecifyAngleAction(ULearningAgentsActionSchema* Schema, const float AngleScale, const bool bActive, const FName Tag)
 {
 	if (!bActive)
 	{
 		return SpecifyNullAction(Schema, Tag);
 	}
 
-	FLearningAgentsActionNormalizationSettings Settings;
-	Settings.Scales = { NormalizationSettings.Scale };
-	Settings.Offsets = { NormalizationSettings.Offset };
-	Settings.NormalizationOption = NormalizationSettings.NormalizationOption;
-
-	return SpecifyContinuousAction(Schema, 1, Settings, bActive, Tag);
+	return SpecifyContinuousAction(Schema, 1, AngleScale, bActive, Tag);
 }
 
-FLearningAgentsActionSchemaElement ULearningAgentsActions::SpecifyVelocityAction(ULearningAgentsActionSchema* Schema, const FLearningAgentsVectorActionNormalizationSettings& NormalizationSettings, const bool bActive, const FName Tag)
+FLearningAgentsActionSchemaElement ULearningAgentsActions::SpecifyVelocityAction(ULearningAgentsActionSchema* Schema, const float VelocityScale, const bool bActive, const FName Tag)
 {
 	if (!bActive)
 	{
 		return SpecifyNullAction(Schema, Tag);
 	}
 
-	FLearningAgentsActionNormalizationSettings Settings;
-	Settings.Scales = { (float)NormalizationSettings.Scales.X, (float)NormalizationSettings.Scales.Y, (float)NormalizationSettings.Scales.Z };
-	Settings.Offsets = { (float)NormalizationSettings.Offsets.X, (float)NormalizationSettings.Offsets.Y, (float)NormalizationSettings.Offsets.Z };
-	Settings.NormalizationOption = NormalizationSettings.NormalizationOption;
-
-	return SpecifyContinuousAction(Schema, 3, Settings, bActive, Tag);
+	return SpecifyContinuousAction(Schema, 3, VelocityScale, bActive, Tag);
 }
 
 FLearningAgentsActionSchemaElement ULearningAgentsActions::SpecifyDirectionAction(ULearningAgentsActionSchema* Schema, const bool bActive, const FName Tag)
@@ -2108,7 +2072,7 @@ FLearningAgentsActionSchemaElement ULearningAgentsActions::SpecifyDirectionActio
 		return SpecifyNullAction(Schema, Tag);
 	}
 
-	return SpecifyContinuousAction(Schema, 3, FLearningAgentsActionNormalizationSettings(), bActive, Tag);
+	return SpecifyContinuousAction(Schema, 3, 1.0f, bActive, Tag);
 }
 
 void ULearningAgentsActions::LogAction(const ULearningAgentsActionObject* Object, const FLearningAgentsActionObjectElement Element)
@@ -2136,24 +2100,32 @@ FLearningAgentsActionObjectElement ULearningAgentsActions::MakeNullAction(ULearn
 FLearningAgentsActionObjectElement ULearningAgentsActions::MakeContinuousAction(
 	ULearningAgentsActionObject* Object,
 	const TArray<float>& Values,
-	const bool bActive,
+    const bool bActive,	
 	const FName Tag,
-	const FLearningAgentsVisualLoggerSettings& VisualLoggerSettings)
+	const bool bVisualLoggerEnabled,
+	ULearningAgentsManagerListener* VisualLoggerListener,
+	const int32 VisualLoggerAgentId,
+	const FVector VisualLoggerLocation,
+	const FLinearColor VisualLoggerColor)
 {
 	if (!bActive)
 	{
 		return MakeNullAction(Object, Tag);
 	}
 
-	return MakeContinuousActionFromArrayView(Object, Values, bActive, Tag, VisualLoggerSettings);
+	return MakeContinuousActionFromArrayView(Object, Values, bActive, Tag, bVisualLoggerEnabled, VisualLoggerListener, VisualLoggerAgentId, VisualLoggerLocation, VisualLoggerColor);
 }
 
 FLearningAgentsActionObjectElement ULearningAgentsActions::MakeContinuousActionFromArrayView(
 	ULearningAgentsActionObject* Object,
 	const TArrayView<const float> Values,
-	const bool bActive,
+    const bool bActive,
 	const FName Tag,
-	const FLearningAgentsVisualLoggerSettings& VisualLoggerSettings)
+	const bool bVisualLoggerEnabled,
+	ULearningAgentsManagerListener* VisualLoggerListener,
+	const int32 VisualLoggerAgentId,
+	const FVector VisualLoggerLocation,
+	const FLinearColor VisualLoggerColor)
 {
 	if (!bActive)
 	{
@@ -2172,16 +2144,16 @@ FLearningAgentsActionObjectElement ULearningAgentsActions::MakeContinuousActionF
 	}
 
 #if UE_LEARNING_AGENTS_ENABLE_VISUAL_LOG
-	if (VisualLoggerSettings.bEnabled && VisualLoggerSettings.Listener)
+	if (bVisualLoggerEnabled && VisualLoggerListener)
 	{
-		const ULearningAgentsVisualLoggerObject* VisualLoggerObject = VisualLoggerSettings.Listener->GetOrAddVisualLoggerObject(Tag);
+		const ULearningAgentsVisualLoggerObject* VisualLoggerObject = VisualLoggerListener->GetOrAddVisualLoggerObject(Tag);
 
-		UE_LEARNING_AGENTS_VLOG_STRING(VisualLoggerObject, LogLearning, Display, VisualLoggerSettings.Location,
-			VisualLoggerSettings.Color.ToFColor(true),
+		UE_LEARNING_AGENTS_VLOG_STRING(VisualLoggerObject, LogLearning, Display, VisualLoggerLocation,
+			VisualLoggerColor.ToFColor(true),
 			TEXT("Listener: %s\nTag: %s\nAgent Id: % 3i\nEncoded: %s\nValues: %s"),
-			*VisualLoggerSettings.Listener->GetName(),
+			*VisualLoggerListener->GetName(),
 			*Tag.ToString(),
-			VisualLoggerSettings.AgentId,
+			VisualLoggerAgentId,
 			*UE::Learning::Array::FormatFloat(Values),
 			*UE::Learning::Array::FormatFloat(Values)); // Encoded is identical to provided values
 	}
@@ -2193,9 +2165,13 @@ FLearningAgentsActionObjectElement ULearningAgentsActions::MakeContinuousActionF
 FLearningAgentsActionObjectElement ULearningAgentsActions::MakeExclusiveDiscreteAction(
 	ULearningAgentsActionObject* Object,
 	const int32 Index,
-	const bool bActive,
+    const bool bActive,
 	const FName Tag,
-	const FLearningAgentsVisualLoggerSettings& VisualLoggerSettings)
+	const bool bVisualLoggerEnabled,
+	ULearningAgentsManagerListener* VisualLoggerListener,
+	const int32 VisualLoggerAgentId,
+	const FVector VisualLoggerLocation,
+	const FLinearColor VisualLoggerColor)
 {
 	if (!bActive)
 	{
@@ -2215,16 +2191,16 @@ FLearningAgentsActionObjectElement ULearningAgentsActions::MakeExclusiveDiscrete
 	}
 
 #if UE_LEARNING_AGENTS_ENABLE_VISUAL_LOG
-	if (VisualLoggerSettings.bEnabled && VisualLoggerSettings.Listener)
+	if (bVisualLoggerEnabled && VisualLoggerListener)
 	{
-		const ULearningAgentsVisualLoggerObject* VisualLoggerObject = VisualLoggerSettings.Listener->GetOrAddVisualLoggerObject(Tag);
+		const ULearningAgentsVisualLoggerObject* VisualLoggerObject = VisualLoggerListener->GetOrAddVisualLoggerObject(Tag);
 
-		UE_LEARNING_AGENTS_VLOG_STRING(VisualLoggerObject, LogLearning, Display, VisualLoggerSettings.Location,
-			VisualLoggerSettings.Color.ToFColor(true),
+		UE_LEARNING_AGENTS_VLOG_STRING(VisualLoggerObject, LogLearning, Display, VisualLoggerLocation,
+			VisualLoggerColor.ToFColor(true),
 			TEXT("Listener: %s\nTag: %s\nAgent Id: % 3i\nIndex: [%i]"),
-			*VisualLoggerSettings.Listener->GetName(),
+			*VisualLoggerListener->GetName(),
 			*Tag.ToString(),
-			VisualLoggerSettings.AgentId,
+			VisualLoggerAgentId,
 			Index);
 	}
 #endif
@@ -2235,9 +2211,13 @@ FLearningAgentsActionObjectElement ULearningAgentsActions::MakeExclusiveDiscrete
 FLearningAgentsActionObjectElement ULearningAgentsActions::MakeNamedExclusiveDiscreteAction(
 	ULearningAgentsActionObject* Object,
 	const FName Name,
-	const bool bActive,
+    const bool bActive,	
 	const FName Tag,
-	const FLearningAgentsVisualLoggerSettings& VisualLoggerSettings)
+	const bool bVisualLoggerEnabled,
+	ULearningAgentsManagerListener* VisualLoggerListener,
+	const int32 VisualLoggerAgentId,
+	const FVector VisualLoggerLocation,
+	const FLinearColor VisualLoggerColor)
 {
 	if (!bActive)
 	{
@@ -2251,16 +2231,16 @@ FLearningAgentsActionObjectElement ULearningAgentsActions::MakeNamedExclusiveDis
 	}
 
 #if UE_LEARNING_AGENTS_ENABLE_VISUAL_LOG
-	if (VisualLoggerSettings.bEnabled && VisualLoggerSettings.Listener)
+	if (bVisualLoggerEnabled && VisualLoggerListener)
 	{
-		const ULearningAgentsVisualLoggerObject* VisualLoggerObject = VisualLoggerSettings.Listener->GetOrAddVisualLoggerObject(Tag);
+		const ULearningAgentsVisualLoggerObject* VisualLoggerObject = VisualLoggerListener->GetOrAddVisualLoggerObject(Tag);
 
-		UE_LEARNING_AGENTS_VLOG_STRING(VisualLoggerObject, LogLearning, Display, VisualLoggerSettings.Location,
-			VisualLoggerSettings.Color.ToFColor(true),
+		UE_LEARNING_AGENTS_VLOG_STRING(VisualLoggerObject, LogLearning, Display, VisualLoggerLocation,
+			VisualLoggerColor.ToFColor(true),
 			TEXT("Listener: %s\nTag: %s\nAgent Id: % 3i\nName: [%s]"),
-			*VisualLoggerSettings.Listener->GetName(),
+			*VisualLoggerListener->GetName(),
 			*Tag.ToString(),
-			VisualLoggerSettings.AgentId,
+			VisualLoggerAgentId,
 			*Name.ToString());
 	}
 #endif
@@ -2271,24 +2251,32 @@ FLearningAgentsActionObjectElement ULearningAgentsActions::MakeNamedExclusiveDis
 FLearningAgentsActionObjectElement ULearningAgentsActions::MakeInclusiveDiscreteAction(
 	ULearningAgentsActionObject* Object,
 	const TArray<int32>& Indices,
-	const bool bActive,
+    const bool bActive,
 	const FName Tag,
-	const FLearningAgentsVisualLoggerSettings& VisualLoggerSettings)
+	const bool bVisualLoggerEnabled,
+	ULearningAgentsManagerListener* VisualLoggerListener,
+	const int32 VisualLoggerAgentId,
+	const FVector VisualLoggerLocation,
+	const FLinearColor VisualLoggerColor)
 {
 	if (!bActive)
 	{
 		return MakeNullAction(Object, Tag);
 	}
 
-	return MakeInclusiveDiscreteActionFromArrayView(Object, Indices, bActive, Tag, VisualLoggerSettings);
+	return MakeInclusiveDiscreteActionFromArrayView(Object, Indices, bActive, Tag, bVisualLoggerEnabled, VisualLoggerListener, VisualLoggerAgentId, VisualLoggerLocation, VisualLoggerColor);
 }
 
 FLearningAgentsActionObjectElement ULearningAgentsActions::MakeInclusiveDiscreteActionFromArrayView(
 	ULearningAgentsActionObject* Object,
 	const TArrayView<const int32> Indices,
-	const bool bActive,
+    const bool bActive,	
 	const FName Tag,
-	const FLearningAgentsVisualLoggerSettings& VisualLoggerSettings)
+	const bool bVisualLoggerEnabled,
+	ULearningAgentsManagerListener* VisualLoggerListener,
+	const int32 VisualLoggerAgentId,
+	const FVector VisualLoggerLocation,
+	const FLinearColor VisualLoggerColor)
 {
 	if (!bActive)
 	{
@@ -2319,16 +2307,16 @@ FLearningAgentsActionObjectElement ULearningAgentsActions::MakeInclusiveDiscrete
 	}
 
 #if UE_LEARNING_AGENTS_ENABLE_VISUAL_LOG
-	if (VisualLoggerSettings.bEnabled && VisualLoggerSettings.Listener)
+	if (bVisualLoggerEnabled && VisualLoggerListener)
 	{
-		const ULearningAgentsVisualLoggerObject* VisualLoggerObject = VisualLoggerSettings.Listener->GetOrAddVisualLoggerObject(Tag);
+		const ULearningAgentsVisualLoggerObject* VisualLoggerObject = VisualLoggerListener->GetOrAddVisualLoggerObject(Tag);
 
-		UE_LEARNING_AGENTS_VLOG_STRING(VisualLoggerObject, LogLearning, Display, VisualLoggerSettings.Location,
-			VisualLoggerSettings.Color.ToFColor(true),
+		UE_LEARNING_AGENTS_VLOG_STRING(VisualLoggerObject, LogLearning, Display, VisualLoggerLocation,
+			VisualLoggerColor.ToFColor(true),
 			TEXT("Listener: %s\nTag: %s\nAgent Id: % 3i\nIndices: %s"),
-			*VisualLoggerSettings.Listener->GetName(),
+			*VisualLoggerListener->GetName(),
 			*Tag.ToString(),
-			VisualLoggerSettings.AgentId,
+			VisualLoggerAgentId,
 			*UE::Learning::Array::FormatInt32(Indices, 256));
 	}
 #endif
@@ -2339,24 +2327,32 @@ FLearningAgentsActionObjectElement ULearningAgentsActions::MakeInclusiveDiscrete
 FLearningAgentsActionObjectElement ULearningAgentsActions::MakeNamedInclusiveDiscreteAction(
 	ULearningAgentsActionObject* Object,
 	const TArray<FName>& Names,
-	const bool bActive,
+    const bool bActive,	
 	const FName Tag,
-	const FLearningAgentsVisualLoggerSettings& VisualLoggerSettings)
+	const bool bVisualLoggerEnabled,
+	ULearningAgentsManagerListener* VisualLoggerListener,
+	const int32 VisualLoggerAgentId,
+	const FVector VisualLoggerLocation,
+	const FLinearColor VisualLoggerColor)
 {
 	if (!bActive)
 	{
 		return MakeNullAction(Object, Tag);
 	}
 
-	return MakeNamedInclusiveDiscreteActionFromArrayView(Object, Names, bActive, Tag, VisualLoggerSettings);
+	return MakeNamedInclusiveDiscreteActionFromArrayView(Object, Names, bActive, Tag, bVisualLoggerEnabled, VisualLoggerListener, VisualLoggerAgentId, VisualLoggerLocation, VisualLoggerColor);
 }
 
 FLearningAgentsActionObjectElement ULearningAgentsActions::MakeNamedInclusiveDiscreteActionFromArrayView(
 	ULearningAgentsActionObject* Object,
 	const TArrayView<const FName> Names,
-	const bool bActive,
+    const bool bActive,
 	const FName Tag,
-	const FLearningAgentsVisualLoggerSettings& VisualLoggerSettings)
+	const bool bVisualLoggerEnabled,
+	ULearningAgentsManagerListener* VisualLoggerListener,
+	const int32 VisualLoggerAgentId,
+	const FVector VisualLoggerLocation,
+	const FLinearColor VisualLoggerColor)
 {
 	if (!bActive)
 	{
@@ -2376,16 +2372,16 @@ FLearningAgentsActionObjectElement ULearningAgentsActions::MakeNamedInclusiveDis
 	}
 
 #if UE_LEARNING_AGENTS_ENABLE_VISUAL_LOG
-	if (VisualLoggerSettings.bEnabled && VisualLoggerSettings.Listener)
+	if (bVisualLoggerEnabled && VisualLoggerListener)
 	{
-		const ULearningAgentsVisualLoggerObject* VisualLoggerObject = VisualLoggerSettings.Listener->GetOrAddVisualLoggerObject(Tag);
+		const ULearningAgentsVisualLoggerObject* VisualLoggerObject = VisualLoggerListener->GetOrAddVisualLoggerObject(Tag);
 
-		UE_LEARNING_AGENTS_VLOG_STRING(VisualLoggerObject, LogLearning, Display, VisualLoggerSettings.Location,
-			VisualLoggerSettings.Color.ToFColor(true),
+		UE_LEARNING_AGENTS_VLOG_STRING(VisualLoggerObject, LogLearning, Display, VisualLoggerLocation,
+			VisualLoggerColor.ToFColor(true),
 			TEXT("Listener: %s\nTag: %s\nAgent Id: % 3i\nNames: %s"),
-			*VisualLoggerSettings.Listener->GetName(),
+			*VisualLoggerListener->GetName(),
 			*Tag.ToString(),
-			VisualLoggerSettings.AgentId,
+			VisualLoggerAgentId,
 			*UE::Learning::Agents::Action::Private::FormatNames(Names, 256));
 	}
 #endif
@@ -2634,9 +2630,13 @@ FLearningAgentsActionObjectElement ULearningAgentsActions::MakeEnumAction(
 	ULearningAgentsActionObject* Object,
 	const UEnum* Enum,
 	const uint8 EnumValue,
-	const bool bActive,
+    const bool bActive,
 	const FName Tag,
-	const FLearningAgentsVisualLoggerSettings& VisualLoggerSettings)
+	const bool bVisualLoggerEnabled,
+	ULearningAgentsManagerListener* VisualLoggerListener,
+	const int32 VisualLoggerAgentId,
+	const FVector VisualLoggerLocation,
+	const FLinearColor VisualLoggerColor)
 {
 	if (!bActive)
 	{
@@ -2653,21 +2653,21 @@ FLearningAgentsActionObjectElement ULearningAgentsActions::MakeEnumAction(
 
 	if (EnumValueIndex == INDEX_NONE || EnumValueIndex < 0 || EnumValueIndex >= Enum->NumEnums() - 1)
 	{
-		UE_LOG(LogLearning, Error, TEXT("MakeEnumAction: EnumValue %i not valid for Enum '%s'."), EnumValue, *Enum->GetName());
+		UE_LOG(LogLearning, Error, TEXT("MakeEnumAction: EnumValue %i not valid for Enum '%s'."), EnumValue , *Enum->GetName());
 		return FLearningAgentsActionObjectElement();
 	}
 
 #if UE_LEARNING_AGENTS_ENABLE_VISUAL_LOG
-	if (VisualLoggerSettings.bEnabled && VisualLoggerSettings.Listener)
+	if (bVisualLoggerEnabled && VisualLoggerListener)
 	{
-		const ULearningAgentsVisualLoggerObject* VisualLoggerObject = VisualLoggerSettings.Listener->GetOrAddVisualLoggerObject(Tag);
+		const ULearningAgentsVisualLoggerObject* VisualLoggerObject = VisualLoggerListener->GetOrAddVisualLoggerObject(Tag);
 
-		UE_LEARNING_AGENTS_VLOG_STRING(VisualLoggerObject, LogLearning, Display, VisualLoggerSettings.Location,
-			VisualLoggerSettings.Color.ToFColor(true),
+		UE_LEARNING_AGENTS_VLOG_STRING(VisualLoggerObject, LogLearning, Display, VisualLoggerLocation,
+			VisualLoggerColor.ToFColor(true),
 			TEXT("Listener: %s\nTag: %s\nAgent Id: % 3i\nEnum: %s\nSize: [%i]\nValue: [%s]\nIndex: [%i]"),
-			*VisualLoggerSettings.Listener->GetName(),
+			*VisualLoggerListener->GetName(),
 			*Tag.ToString(),
-			VisualLoggerSettings.AgentId,
+			VisualLoggerAgentId,
 			*Enum->GetName(),
 			Enum->NumEnums() - 1,
 			*Enum->GetDisplayNameTextByValue(EnumValue).ToString(),
@@ -2682,9 +2682,13 @@ FLearningAgentsActionObjectElement ULearningAgentsActions::MakeBitmaskAction(
 	ULearningAgentsActionObject* Object,
 	const UEnum* Enum,
 	const int32 BitmaskValue,
-	const bool bActive,
+    const bool bActive,	
 	const FName Tag,
-	const FLearningAgentsVisualLoggerSettings& VisualLoggerSettings)
+	const bool bVisualLoggerEnabled,
+	ULearningAgentsManagerListener* VisualLoggerListener,
+	const int32 VisualLoggerAgentId,
+	const FVector VisualLoggerLocation,
+	const FLinearColor VisualLoggerColor)
 {
 	if (!bActive)
 	{
@@ -2715,9 +2719,9 @@ FLearningAgentsActionObjectElement ULearningAgentsActions::MakeBitmaskAction(
 	}
 
 #if UE_LEARNING_AGENTS_ENABLE_VISUAL_LOG
-	if (VisualLoggerSettings.bEnabled && VisualLoggerSettings.Listener)
+	if (bVisualLoggerEnabled && VisualLoggerListener)
 	{
-		const ULearningAgentsVisualLoggerObject* VisualLoggerObject = VisualLoggerSettings.Listener->GetOrAddVisualLoggerObject(Tag);
+		const ULearningAgentsVisualLoggerObject* VisualLoggerObject = VisualLoggerListener->GetOrAddVisualLoggerObject(Tag);
 
 		FString ValuesString;
 		FString IndicesString;
@@ -2734,12 +2738,12 @@ FLearningAgentsActionObjectElement ULearningAgentsActions::MakeBitmaskAction(
 		ValuesString = ValuesString.TrimEnd();
 		IndicesString = IndicesString.TrimEnd();
 
-		UE_LEARNING_AGENTS_VLOG_STRING(VisualLoggerObject, LogLearning, Display, VisualLoggerSettings.Location,
-			VisualLoggerSettings.Color.ToFColor(true),
+		UE_LEARNING_AGENTS_VLOG_STRING(VisualLoggerObject, LogLearning, Display, VisualLoggerLocation,
+			VisualLoggerColor.ToFColor(true),
 			TEXT("Listener: %s\nTag: %s\nAgent Id: % 3i\nEnum: %s\nSize: [%i]\nValues: [%s]\nIndices: [%s]"),
-			*VisualLoggerSettings.Listener->GetName(),
+			*VisualLoggerListener->GetName(),
 			*Tag.ToString(),
-			VisualLoggerSettings.AgentId,
+			VisualLoggerAgentId,
 			*Enum->GetName(),
 			Enum->NumEnums() - 1,
 			*ValuesString,
@@ -2838,9 +2842,13 @@ FLearningAgentsActionObjectElement ULearningAgentsActions::MakeEncodingAction(UL
 FLearningAgentsActionObjectElement ULearningAgentsActions::MakeBoolAction(
 	ULearningAgentsActionObject* Object,
 	const bool bValue,
-	const bool bActive,
+    const bool bActive,
 	const FName Tag,
-	const FLearningAgentsVisualLoggerSettings& VisualLoggerSettings)
+	const bool bVisualLoggerEnabled,
+	ULearningAgentsManagerListener* VisualLoggerListener,
+	const int32 VisualLoggerAgentId,
+	const FVector VisualLoggerLocation,
+	const FLinearColor VisualLoggerColor)
 {
 	if (!bActive)
 	{
@@ -2848,16 +2856,16 @@ FLearningAgentsActionObjectElement ULearningAgentsActions::MakeBoolAction(
 	}
 
 #if UE_LEARNING_AGENTS_ENABLE_VISUAL_LOG
-	if (VisualLoggerSettings.bEnabled && VisualLoggerSettings.Listener)
+	if (bVisualLoggerEnabled && VisualLoggerListener)
 	{
-		const ULearningAgentsVisualLoggerObject* VisualLoggerObject = VisualLoggerSettings.Listener->GetOrAddVisualLoggerObject(Tag);
+		const ULearningAgentsVisualLoggerObject* VisualLoggerObject = VisualLoggerListener->GetOrAddVisualLoggerObject(Tag);
 
-		UE_LEARNING_AGENTS_VLOG_STRING(VisualLoggerObject, LogLearning, Display, VisualLoggerSettings.Location,
-			VisualLoggerSettings.Color.ToFColor(true),
+		UE_LEARNING_AGENTS_VLOG_STRING(VisualLoggerObject, LogLearning, Display, VisualLoggerLocation,
+			VisualLoggerColor.ToFColor(true),
 			TEXT("Listener: %s\nTag: %s\nAgent Id: % 3i\nValue: [%s]"),
-			*VisualLoggerSettings.Listener->GetName(),
+			*VisualLoggerListener->GetName(),
 			*Tag.ToString(),
-			VisualLoggerSettings.AgentId,
+			VisualLoggerAgentId,
 			bValue ? TEXT("true") : TEXT("false"));
 	}
 #endif
@@ -2868,9 +2876,13 @@ FLearningAgentsActionObjectElement ULearningAgentsActions::MakeBoolAction(
 FLearningAgentsActionObjectElement ULearningAgentsActions::MakeFloatAction(
 	ULearningAgentsActionObject* Object,
 	const float Value,
-	const bool bActive,
+    const bool bActive,
 	const FName Tag,
-	const FLearningAgentsVisualLoggerSettings& VisualLoggerSettings)
+	const bool bVisualLoggerEnabled,
+	ULearningAgentsManagerListener* VisualLoggerListener,
+	const int32 VisualLoggerAgentId,
+	const FVector VisualLoggerLocation,
+	const FLinearColor VisualLoggerColor)
 {
 	if (!bActive)
 	{
@@ -2878,30 +2890,34 @@ FLearningAgentsActionObjectElement ULearningAgentsActions::MakeFloatAction(
 	}
 
 #if UE_LEARNING_AGENTS_ENABLE_VISUAL_LOG
-	if (VisualLoggerSettings.bEnabled && VisualLoggerSettings.Listener)
+	if (bVisualLoggerEnabled && VisualLoggerListener)
 	{
-		const ULearningAgentsVisualLoggerObject* VisualLoggerObject = VisualLoggerSettings.Listener->GetOrAddVisualLoggerObject(Tag);
+		const ULearningAgentsVisualLoggerObject* VisualLoggerObject = VisualLoggerListener->GetOrAddVisualLoggerObject(Tag);
 
-		UE_LEARNING_AGENTS_VLOG_STRING(VisualLoggerObject, LogLearning, Display, VisualLoggerSettings.Location,
-			VisualLoggerSettings.Color.ToFColor(true),
+		UE_LEARNING_AGENTS_VLOG_STRING(VisualLoggerObject, LogLearning, Display, VisualLoggerLocation,
+			VisualLoggerColor.ToFColor(true),
 			TEXT("Listener: %s\nTag: %s\nAgent Id: % 3i\nValue: [% 6.1f]"),
-			*VisualLoggerSettings.Listener->GetName(),
+			*VisualLoggerListener->GetName(),
 			*Tag.ToString(),
-			VisualLoggerSettings.AgentId,
+			VisualLoggerAgentId,
 			Value);
 	}
 #endif
 
-	return MakeContinuousActionFromArrayView(Object, { Value }, bActive, Tag, FLearningAgentsVisualLoggerSettings());
+	return MakeContinuousActionFromArrayView(Object, { Value }, bActive, Tag);
 }
 
 FLearningAgentsActionObjectElement ULearningAgentsActions::MakeLocationAction(
 	ULearningAgentsActionObject* Object,
 	const FVector Location,
 	const FTransform RelativeTransform,
-	const bool bActive,
+	const bool bActive,	
 	const FName Tag,
-	const FLearningAgentsVisualLoggerSettings& VisualLoggerSettings)
+	const bool bVisualLoggerEnabled,
+	ULearningAgentsManagerListener* VisualLoggerListener,
+	const int32 VisualLoggerAgentId,
+	const FVector VisualLoggerLocation,
+	const FLinearColor VisualLoggerColor)
 {
 	if (!bActive)
 	{
@@ -2911,34 +2927,34 @@ FLearningAgentsActionObjectElement ULearningAgentsActions::MakeLocationAction(
 	const FVector LocalLocation = RelativeTransform.InverseTransformPosition(Location);
 
 #if UE_LEARNING_AGENTS_ENABLE_VISUAL_LOG
-	if (VisualLoggerSettings.bEnabled && VisualLoggerSettings.Listener)
+	if (bVisualLoggerEnabled && VisualLoggerListener)
 	{
-		const ULearningAgentsVisualLoggerObject* VisualLoggerObject = VisualLoggerSettings.Listener->GetOrAddVisualLoggerObject(Tag);
+		const ULearningAgentsVisualLoggerObject* VisualLoggerObject = VisualLoggerListener->GetOrAddVisualLoggerObject(Tag);
 
 		UE_LEARNING_AGENTS_VLOG_LOCATION(VisualLoggerObject, LogLearning, Display,
 			Location,
 			10,
-			VisualLoggerSettings.Color.ToFColor(true),
+			VisualLoggerColor.ToFColor(true),
 			TEXT(""));
 
 		UE_LEARNING_AGENTS_VLOG_SEGMENT(VisualLoggerObject, LogLearning, Display,
 			RelativeTransform.GetTranslation(),
 			Location,
-			VisualLoggerSettings.Color.ToFColor(true),
+			VisualLoggerColor.ToFColor(true),
 			TEXT(""));
 
 		UE_LEARNING_AGENTS_VLOG_TRANSFORM(VisualLoggerObject, LogLearning, Display,
 			RelativeTransform.GetTranslation(),
 			RelativeTransform.GetRotation(),
-			VisualLoggerSettings.Color.ToFColor(true),
+			VisualLoggerColor.ToFColor(true),
 			TEXT(""));
 
-		UE_LEARNING_AGENTS_VLOG_STRING(VisualLoggerObject, LogLearning, Display, VisualLoggerSettings.Location,
-			VisualLoggerSettings.Color.ToFColor(true),
+		UE_LEARNING_AGENTS_VLOG_STRING(VisualLoggerObject, LogLearning, Display, VisualLoggerLocation,
+			VisualLoggerColor.ToFColor(true),
 			TEXT("Listener: %s\nTag: %s\nAgent Id: % 3i\nLocal Location: [% 6.1f % 6.1f % 6.1f]\nLocation: [% 6.1f % 6.1f % 6.1f]"),
-			*VisualLoggerSettings.Listener->GetName(),
+			*VisualLoggerListener->GetName(),
 			*Tag.ToString(),
-			VisualLoggerSettings.AgentId,
+			VisualLoggerAgentId,
 			LocalLocation.X, LocalLocation.Y, LocalLocation.Z,
 			Location.X, Location.Y, Location.Z);
 	}
@@ -2947,16 +2963,21 @@ FLearningAgentsActionObjectElement ULearningAgentsActions::MakeLocationAction(
 	return MakeContinuousActionFromArrayView(Object, {
 		(float)LocalLocation.X,
 		(float)LocalLocation.Y,
-		(float)LocalLocation.Z }, bActive, Tag, FLearningAgentsVisualLoggerSettings());
+		(float)LocalLocation.Z }, bActive, Tag);
 }
 
 FLearningAgentsActionObjectElement ULearningAgentsActions::MakeRotationAction(
 	ULearningAgentsActionObject* Object,
 	const FRotator Rotation,
 	const FRotator RelativeRotation,
-	const bool bActive,
+    const bool bActive,	
 	const FName Tag,
-	const FLearningAgentsVisualLoggerSettings& VisualLoggerSettings)
+	const bool bVisualLoggerEnabled,
+	ULearningAgentsManagerListener* VisualLoggerListener,
+	const int32 VisualLoggerAgentId,
+	const FVector VisualLoggerRotationLocation,
+	const FVector VisualLoggerLocation,
+	const FLinearColor VisualLoggerColor)
 {
 	if (!bActive)
 	{
@@ -2964,7 +2985,7 @@ FLearningAgentsActionObjectElement ULearningAgentsActions::MakeRotationAction(
 	}
 
 	return MakeRotationActionFromQuat(Object, FQuat::MakeFromRotator(Rotation), FQuat::MakeFromRotator(RelativeRotation), bActive, Tag,
-		VisualLoggerSettings);
+		bVisualLoggerEnabled, VisualLoggerListener, VisualLoggerAgentId, VisualLoggerRotationLocation, VisualLoggerLocation, VisualLoggerColor);
 }
 
 FLearningAgentsActionObjectElement ULearningAgentsActions::MakeRotationActionFromQuat(
@@ -2973,7 +2994,12 @@ FLearningAgentsActionObjectElement ULearningAgentsActions::MakeRotationActionFro
 	const FQuat RelativeRotation,
 	const bool bActive,
 	const FName Tag,
-	const FLearningAgentsVisualLoggerSettings& VisualLoggerSettings)
+	const bool bVisualLoggerEnabled,
+	ULearningAgentsManagerListener* VisualLoggerListener,
+	const int32 VisualLoggerAgentId,
+	const FVector VisualLoggerRotationLocation,
+	const FVector VisualLoggerLocation,
+	const FLinearColor VisualLoggerColor)
 {
 	if (!bActive)
 	{
@@ -2984,22 +3010,22 @@ FLearningAgentsActionObjectElement ULearningAgentsActions::MakeRotationActionFro
 	const FVector RotationVector = LocalRotation.ToRotationVector();
 
 #if UE_LEARNING_AGENTS_ENABLE_VISUAL_LOG
-	if (VisualLoggerSettings.bEnabled && VisualLoggerSettings.Listener)
+	if (bVisualLoggerEnabled && VisualLoggerListener)
 	{
-		const ULearningAgentsVisualLoggerObject* VisualLoggerObject = VisualLoggerSettings.Listener->GetOrAddVisualLoggerObject(Tag);
+		const ULearningAgentsVisualLoggerObject* VisualLoggerObject = VisualLoggerListener->GetOrAddVisualLoggerObject(Tag);
 
 		UE_LEARNING_AGENTS_VLOG_TRANSFORM(VisualLoggerObject, LogLearning, Display,
-			VisualLoggerSettings.DebugLocation,
+			VisualLoggerRotationLocation,
 			LocalRotation.Rotator(),
-			VisualLoggerSettings.Color.ToFColor(true),
+			VisualLoggerColor.ToFColor(true),
 			TEXT(""));
 
-		UE_LEARNING_AGENTS_VLOG_STRING(VisualLoggerObject, LogLearning, Display, VisualLoggerSettings.Location,
-			VisualLoggerSettings.Color.ToFColor(true),
+		UE_LEARNING_AGENTS_VLOG_STRING(VisualLoggerObject, LogLearning, Display, VisualLoggerLocation,
+			VisualLoggerColor.ToFColor(true),
 			TEXT("Listener: %s\nTag: %s\nAgent Id: % 3i\nLocal Rotation Vector: [% 6.1f % 6.1f % 6.1f]\nLocal Rotation: [% 6.1f % 6.1f % 6.1f % 6.1f]\nRotation: [% 6.1f % 6.1f % 6.1f % 6.1f]"),
-			*VisualLoggerSettings.Listener->GetName(),
+			*VisualLoggerListener->GetName(),
 			*Tag.ToString(),
-			VisualLoggerSettings.AgentId,
+			VisualLoggerAgentId,
 			RotationVector.X, RotationVector.Y, RotationVector.Z,
 			LocalRotation.X, LocalRotation.Y, LocalRotation.Z, LocalRotation.W,
 			Rotation.X, Rotation.Y, Rotation.Z, Rotation.W);
@@ -3010,16 +3036,20 @@ FLearningAgentsActionObjectElement ULearningAgentsActions::MakeRotationActionFro
 		(float)FMath::RadiansToDegrees(RotationVector.X),
 		(float)FMath::RadiansToDegrees(RotationVector.Y),
 		(float)FMath::RadiansToDegrees(RotationVector.Z),
-		}, bActive, Tag, FLearningAgentsVisualLoggerSettings());
+		}, bActive, Tag);
 }
 
 FLearningAgentsActionObjectElement ULearningAgentsActions::MakeScaleAction(
 	ULearningAgentsActionObject* Object,
 	const FVector Scale,
 	const FVector RelativeScale,
-	const bool bActive,
+    const bool bActive,	
 	const FName Tag,
-	const FLearningAgentsVisualLoggerSettings& VisualLoggerSettings)
+	const bool bVisualLoggerEnabled,
+	ULearningAgentsManagerListener* VisualLoggerListener,
+	const int32 VisualLoggerAgentId,
+	const FVector VisualLoggerLocation,
+	const FLinearColor VisualLoggerColor)
 {
 	if (!bActive)
 	{
@@ -3031,16 +3061,16 @@ FLearningAgentsActionObjectElement ULearningAgentsActions::MakeScaleAction(
 		UE::Learning::Agents::Action::Private::VectorLogSafe(RelativeScale);
 
 #if UE_LEARNING_AGENTS_ENABLE_VISUAL_LOG
-	if (VisualLoggerSettings.bEnabled && VisualLoggerSettings.Listener)
+	if (bVisualLoggerEnabled && VisualLoggerListener)
 	{
-		const ULearningAgentsVisualLoggerObject* VisualLoggerObject = VisualLoggerSettings.Listener->GetOrAddVisualLoggerObject(Tag);
+		const ULearningAgentsVisualLoggerObject* VisualLoggerObject = VisualLoggerListener->GetOrAddVisualLoggerObject(Tag);
 
-		UE_LEARNING_AGENTS_VLOG_STRING(VisualLoggerObject, LogLearning, Display, VisualLoggerSettings.Location,
-			VisualLoggerSettings.Color.ToFColor(true),
+		UE_LEARNING_AGENTS_VLOG_STRING(VisualLoggerObject, LogLearning, Display, VisualLoggerLocation,
+			VisualLoggerColor.ToFColor(true),
 			TEXT("Listener: %s\nTag: %s\nAgent Id: % 3i\nLocal Log Scale: [% 6.1f % 6.1f % 6.1f]\nScale: [% 6.1f % 6.1f % 6.1f]"),
-			*VisualLoggerSettings.Listener->GetName(),
+			*VisualLoggerListener->GetName(),
 			*Tag.ToString(),
-			VisualLoggerSettings.AgentId,
+			VisualLoggerAgentId,
 			LocalLogScale.X, LocalLogScale.Y, LocalLogScale.Z,
 			Scale.X, Scale.Y, Scale.Z);
 	}
@@ -3050,16 +3080,20 @@ FLearningAgentsActionObjectElement ULearningAgentsActions::MakeScaleAction(
 		(float)LocalLogScale.X,
 		(float)LocalLogScale.Y,
 		(float)LocalLogScale.Z,
-		}, bActive, Tag, FLearningAgentsVisualLoggerSettings());
+		}, bActive, Tag);
 }
 
 FLearningAgentsActionObjectElement ULearningAgentsActions::MakeTransformAction(
 	ULearningAgentsActionObject* Object,
 	const FTransform Transform,
 	const FTransform RelativeTransform,
-	const bool bActive,
+	const bool bActive,	
 	const FName Tag,
-	const FLearningAgentsVisualLoggerSettings& VisualLoggerSettings)
+	const bool bVisualLoggerEnabled,
+	ULearningAgentsManagerListener* VisualLoggerListener,
+	const int32 VisualLoggerAgentId,
+	const FVector VisualLoggerLocation,
+	const FLinearColor VisualLoggerColor)
 {
 	if (!bActive)
 	{
@@ -3071,22 +3105,22 @@ FLearningAgentsActionObjectElement ULearningAgentsActions::MakeTransformAction(
 	FVector Scale = Transform.GetScale3D();
 
 #if UE_LEARNING_AGENTS_ENABLE_VISUAL_LOG
-	if (VisualLoggerSettings.bEnabled && VisualLoggerSettings.Listener)
+	if (bVisualLoggerEnabled && VisualLoggerListener)
 	{
-		const ULearningAgentsVisualLoggerObject* VisualLoggerObject = VisualLoggerSettings.Listener->GetOrAddVisualLoggerObject(Tag);
+		const ULearningAgentsVisualLoggerObject* VisualLoggerObject = VisualLoggerListener->GetOrAddVisualLoggerObject(Tag);
 
 		UE_LEARNING_AGENTS_VLOG_TRANSFORM(VisualLoggerObject, LogLearning, Display,
 			Location,
 			Rotation,
-			VisualLoggerSettings.Color.ToFColor(true),
+			VisualLoggerColor.ToFColor(true),
 			TEXT(""));
 
-		UE_LEARNING_AGENTS_VLOG_STRING(VisualLoggerObject, LogLearning, Display, VisualLoggerSettings.Location,
-			VisualLoggerSettings.Color.ToFColor(true),
+		UE_LEARNING_AGENTS_VLOG_STRING(VisualLoggerObject, LogLearning, Display, VisualLoggerLocation,
+			VisualLoggerColor.ToFColor(true),
 			TEXT("Listener: %s\nTag: %s\nAgent Id: % 3i\nLocation: [% 6.1f % 6.1f % 6.1f]\nRotation: [% 6.1f % 6.1f % 6.1f % 6.1f]\nScale: [% 6.1f % 6.1f % 6.1f]"),
-			*VisualLoggerSettings.Listener->GetName(),
+			*VisualLoggerListener->GetName(),
 			*Tag.ToString(),
-			VisualLoggerSettings.AgentId,
+			VisualLoggerAgentId,
 			Location.X, Location.Y, Location.Z,
 			Rotation.X, Rotation.Y, Rotation.Z, Rotation.W,
 			Scale.X, Scale.Y, Scale.Z);
@@ -3113,7 +3147,12 @@ FLearningAgentsActionObjectElement ULearningAgentsActions::MakeAngleAction(
 	const float RelativeAngle,
 	const bool bActive,
 	const FName Tag,
-	const FLearningAgentsVisualLoggerSettings& VisualLoggerSettings)
+	const bool bVisualLoggerEnabled,
+	ULearningAgentsManagerListener* VisualLoggerListener,
+	const int32 VisualLoggerAgentId,
+	const FVector VisualLoggerAngleLocation,
+	const FVector VisualLoggerLocation,
+	const FLinearColor VisualLoggerColor)
 {
 	if (!bActive)
 	{
@@ -3121,41 +3160,46 @@ FLearningAgentsActionObjectElement ULearningAgentsActions::MakeAngleAction(
 	}
 
 	float LocalAngle = FMath::FindDeltaAngleDegrees(RelativeAngle, Angle);
-
+	
 #if UE_LEARNING_AGENTS_ENABLE_VISUAL_LOG
-	if (VisualLoggerSettings.bEnabled && VisualLoggerSettings.Listener)
+	if (bVisualLoggerEnabled && VisualLoggerListener)
 	{
-		const ULearningAgentsVisualLoggerObject* VisualLoggerObject = VisualLoggerSettings.Listener->GetOrAddVisualLoggerObject(Tag);
+		const ULearningAgentsVisualLoggerObject* VisualLoggerObject = VisualLoggerListener->GetOrAddVisualLoggerObject(Tag);
 
 		UE_LEARNING_AGENTS_VLOG_ANGLE_DEGREES(VisualLoggerObject, LogLearning, Display,
 			Angle,
 			0.0f,
-			VisualLoggerSettings.DebugLocation,
+			VisualLoggerAngleLocation,
 			10.0f,
-			VisualLoggerSettings.Color.ToFColor(true),
+			VisualLoggerColor.ToFColor(true),
 			TEXT(""));
 
-		UE_LEARNING_AGENTS_VLOG_STRING(VisualLoggerObject, LogLearning, Display, VisualLoggerSettings.Location,
-			VisualLoggerSettings.Color.ToFColor(true),
+		UE_LEARNING_AGENTS_VLOG_STRING(VisualLoggerObject, LogLearning, Display, VisualLoggerLocation,
+			VisualLoggerColor.ToFColor(true),
 			TEXT("Listener: %s\nTag: %s\nAgent Id: % 3i\nLocal Angle: [% 6.1f]\nAngle: [% 6.1f]"),
-			*VisualLoggerSettings.Listener->GetName(),
+			*VisualLoggerListener->GetName(),
 			*Tag.ToString(),
-			VisualLoggerSettings.AgentId,
+			VisualLoggerAgentId,
 			LocalAngle,
 			Angle);
 	}
 #endif
 
-	return MakeContinuousActionFromArrayView(Object, { LocalAngle }, bActive, Tag, FLearningAgentsVisualLoggerSettings());
+	return MakeContinuousActionFromArrayView(Object, { LocalAngle }, bActive, Tag);
 }
 
 FLearningAgentsActionObjectElement ULearningAgentsActions::MakeAngleActionRadians(
 	ULearningAgentsActionObject* Object,
 	const float Angle,
 	const float RelativeAngle,
-	const bool bActive,
+    const bool bActive,	
 	const FName Tag,
-	const FLearningAgentsVisualLoggerSettings& VisualLoggerSettings)
+	const bool bVisualLoggerEnabled,
+	ULearningAgentsManagerListener* VisualLoggerListener,
+	const int32 VisualLoggerAgentId,
+	const FVector VisualLoggerAngleLocation,
+	const FVector VisualLoggerLocation,
+	const FLinearColor VisualLoggerColor)
 {
 	if (!bActive)
 	{
@@ -3163,16 +3207,21 @@ FLearningAgentsActionObjectElement ULearningAgentsActions::MakeAngleActionRadian
 	}
 
 	return MakeAngleAction(Object, FMath::RadiansToDegrees(Angle), FMath::RadiansToDegrees(RelativeAngle), bActive, Tag,
-		VisualLoggerSettings);
+		bVisualLoggerEnabled, VisualLoggerListener, VisualLoggerAgentId, VisualLoggerAngleLocation, VisualLoggerLocation, VisualLoggerColor);
 }
 
 FLearningAgentsActionObjectElement ULearningAgentsActions::MakeVelocityAction(
 	ULearningAgentsActionObject* Object,
 	const FVector Velocity,
 	const FTransform RelativeTransform,
-	const bool bActive,
+	const bool bActive,	
 	const FName Tag,
-	const FLearningAgentsVisualLoggerSettings& VisualLoggerSettings)
+	const bool bVisualLoggerEnabled,
+	ULearningAgentsManagerListener* VisualLoggerListener,
+	const int32 VisualLoggerAgentId,
+	const FVector VisualLoggerVelocityLocation,
+	const FVector VisualLoggerLocation,
+	const FLinearColor VisualLoggerColor)
 {
 	if (!bActive)
 	{
@@ -3182,28 +3231,28 @@ FLearningAgentsActionObjectElement ULearningAgentsActions::MakeVelocityAction(
 	const FVector LocalVelocity = RelativeTransform.InverseTransformVectorNoScale(Velocity);
 
 #if UE_LEARNING_AGENTS_ENABLE_VISUAL_LOG
-	if (VisualLoggerSettings.bEnabled && VisualLoggerSettings.Listener)
+	if (bVisualLoggerEnabled && VisualLoggerListener)
 	{
-		const ULearningAgentsVisualLoggerObject* VisualLoggerObject = VisualLoggerSettings.Listener->GetOrAddVisualLoggerObject(Tag);
+		const ULearningAgentsVisualLoggerObject* VisualLoggerObject = VisualLoggerListener->GetOrAddVisualLoggerObject(Tag);
 
 		UE_LEARNING_AGENTS_VLOG_ARROW(VisualLoggerObject, LogLearning, Display,
-			VisualLoggerSettings.DebugLocation,
-			VisualLoggerSettings.DebugLocation + Velocity,
-			VisualLoggerSettings.Color.ToFColor(true),
+			VisualLoggerVelocityLocation,
+			VisualLoggerVelocityLocation + Velocity,
+			VisualLoggerColor.ToFColor(true),
 			TEXT(""));
 
 		UE_LEARNING_AGENTS_VLOG_TRANSFORM(VisualLoggerObject, LogLearning, Display,
 			RelativeTransform.GetTranslation(),
 			RelativeTransform.GetRotation(),
-			VisualLoggerSettings.Color.ToFColor(true),
+			VisualLoggerColor.ToFColor(true),
 			TEXT(""));
 
-		UE_LEARNING_AGENTS_VLOG_STRING(VisualLoggerObject, LogLearning, Display, VisualLoggerSettings.Location,
-			VisualLoggerSettings.Color.ToFColor(true),
+		UE_LEARNING_AGENTS_VLOG_STRING(VisualLoggerObject, LogLearning, Display, VisualLoggerLocation,
+			VisualLoggerColor.ToFColor(true),
 			TEXT("Listener: %s\nTag: %s\nAgent Id: % 3i\nLocal Velocity: [% 6.1f % 6.1f % 6.1f]\nVelocity: [% 6.1f % 6.1f % 6.1f]"),
-			*VisualLoggerSettings.Listener->GetName(),
+			*VisualLoggerListener->GetName(),
 			*Tag.ToString(),
-			VisualLoggerSettings.AgentId,
+			VisualLoggerAgentId,
 			LocalVelocity.X, LocalVelocity.Y, LocalVelocity.Z,
 			Velocity.X, Velocity.Y, Velocity.Z);
 	}
@@ -3213,16 +3262,22 @@ FLearningAgentsActionObjectElement ULearningAgentsActions::MakeVelocityAction(
 		(float)LocalVelocity.X,
 		(float)LocalVelocity.Y,
 		(float)LocalVelocity.Z,
-		}, bActive, Tag, FLearningAgentsVisualLoggerSettings());
+		}, bActive, Tag);
 }
 
 FLearningAgentsActionObjectElement ULearningAgentsActions::MakeDirectionAction(
 	ULearningAgentsActionObject* Object,
 	const FVector Direction,
 	const FTransform RelativeTransform,
-	const bool bActive,
+	const bool bActive,	
 	const FName Tag,
-	const FLearningAgentsVisualLoggerSettings& VisualLoggerSettings)
+	const bool bVisualLoggerEnabled,
+	ULearningAgentsManagerListener* VisualLoggerListener,
+	const int32 VisualLoggerAgentId,
+	const FVector VisualLoggerDirectionLocation,
+	const FVector VisualLoggerLocation,
+	const float VisualLoggerArrowLength,
+	const FLinearColor VisualLoggerColor)
 {
 	if (!bActive)
 	{
@@ -3232,28 +3287,28 @@ FLearningAgentsActionObjectElement ULearningAgentsActions::MakeDirectionAction(
 	const FVector LocalDirection = RelativeTransform.InverseTransformVectorNoScale(Direction).GetSafeNormal(UE_SMALL_NUMBER, FVector::ForwardVector);
 
 #if UE_LEARNING_AGENTS_ENABLE_VISUAL_LOG
-	if (VisualLoggerSettings.bEnabled && VisualLoggerSettings.Listener)
+	if (bVisualLoggerEnabled && VisualLoggerListener)
 	{
-		const ULearningAgentsVisualLoggerObject* VisualLoggerObject = VisualLoggerSettings.Listener->GetOrAddVisualLoggerObject(Tag);
+		const ULearningAgentsVisualLoggerObject* VisualLoggerObject = VisualLoggerListener->GetOrAddVisualLoggerObject(Tag);
 
 		UE_LEARNING_AGENTS_VLOG_ARROW(VisualLoggerObject, LogLearning, Display,
-			VisualLoggerSettings.DebugLocation,
-			VisualLoggerSettings.DebugLocation + VisualLoggerSettings.ArrowLength * Direction,
-			VisualLoggerSettings.Color.ToFColor(true),
+			VisualLoggerDirectionLocation,
+			VisualLoggerDirectionLocation + VisualLoggerArrowLength * Direction,
+			VisualLoggerColor.ToFColor(true),
 			TEXT(""));
 
 		UE_LEARNING_AGENTS_VLOG_TRANSFORM(VisualLoggerObject, LogLearning, Display,
 			RelativeTransform.GetTranslation(),
 			RelativeTransform.GetRotation(),
-			VisualLoggerSettings.Color.ToFColor(true),
+			VisualLoggerColor.ToFColor(true),
 			TEXT(""));
 
-		UE_LEARNING_AGENTS_VLOG_STRING(VisualLoggerObject, LogLearning, Display, VisualLoggerSettings.Location,
-			VisualLoggerSettings.Color.ToFColor(true),
+		UE_LEARNING_AGENTS_VLOG_STRING(VisualLoggerObject, LogLearning, Display, VisualLoggerLocation,
+			VisualLoggerColor.ToFColor(true),
 			TEXT("Listener: %s\nTag: %s\nAgent Id: % 3i\nLocal Direction: [% 6.1f % 6.1f % 6.1f]\nDirection: [% 6.1f % 6.1f % 6.1f]"),
-			*VisualLoggerSettings.Listener->GetName(),
+			*VisualLoggerListener->GetName(),
 			*Tag.ToString(),
-			VisualLoggerSettings.AgentId,
+			VisualLoggerAgentId,
 			LocalDirection.X, LocalDirection.Y, LocalDirection.Z,
 			Direction.X, Direction.Y, Direction.Z);
 	}
@@ -3263,7 +3318,7 @@ FLearningAgentsActionObjectElement ULearningAgentsActions::MakeDirectionAction(
 		(float)LocalDirection.X,
 		(float)LocalDirection.Y,
 		(float)LocalDirection.Z,
-		}, bActive, Tag, FLearningAgentsVisualLoggerSettings());
+		}, bActive, Tag);
 }
 
 FLearningAgentsActionModifierElement ULearningAgentsActions::MakeNullActionModifier(ULearningAgentsActionModifier* Modifier, const FName Tag)
@@ -3316,7 +3371,7 @@ FLearningAgentsActionModifierElement ULearningAgentsActions::MakeExclusiveDiscre
 	{
 		UE_LOG(LogLearning, Error, TEXT("MakeExclusiveDiscreteActionModifierFromArrayView: MaskedIndices contain duplicates."));
 		return FLearningAgentsActionModifierElement();
-	}
+	}	
 
 	return { Modifier->ActionModifier.CreateDiscreteExclusive({ MaskedIndices }, Tag) };
 }
@@ -3360,7 +3415,7 @@ FLearningAgentsActionModifierElement ULearningAgentsActions::MakeInclusiveDiscre
 	{
 		UE_LOG(LogLearning, Error, TEXT("MakeNamedExclusiveDiscreteActionModifierFromArrayView: MaskedIndices contain duplicates."));
 		return FLearningAgentsActionModifierElement();
-	}
+	}	
 
 	return { Modifier->ActionModifier.CreateDiscreteInclusive({ MaskedIndices }, Tag) };
 }
@@ -3377,12 +3432,12 @@ FLearningAgentsActionModifierElement ULearningAgentsActions::MakeNamedInclusiveD
 		UE_LOG(LogLearning, Error, TEXT("MakeNamedInclusiveDiscreteActionModifierFromArrayView: Modifier is nullptr."));
 		return FLearningAgentsActionModifierElement();
 	}
-
+	
 	if (UE::Learning::Agents::Action::Private::ContainsDuplicates(MaskedNames))
 	{
 		UE_LOG(LogLearning, Error, TEXT("MakeNamedExclusiveDiscreteActionModifierFromArrayView: MaskedNames contain duplicates."));
 		return FLearningAgentsActionModifierElement();
-	}
+	}	
 
 	return { Modifier->ActionModifier.CreateNamedDiscreteInclusive({ MaskedNames }, Tag) };
 }
@@ -3750,8 +3805,8 @@ FLearningAgentsActionModifierElement ULearningAgentsActions::MakeLocationActionM
 		bMaskedY,
 		bMaskedZ
 		}, {
-		(float)LocalLocation.X,
-		(float)LocalLocation.Y,
+		(float)LocalLocation.X, 
+		(float)LocalLocation.Y, 
 		(float)LocalLocation.Z }, Tag);
 }
 
@@ -3766,8 +3821,8 @@ FLearningAgentsActionModifierElement ULearningAgentsActions::MakeScaleActionModi
 		bMaskedY,
 		bMaskedZ
 		}, {
-		(float)LocalLogScale.X,
-		(float)LocalLogScale.Y,
+		(float)LocalLogScale.X, 
+		(float)LocalLogScale.Y, 
 		(float)LocalLogScale.Z }, Tag);
 }
 
@@ -3881,12 +3936,16 @@ bool ULearningAgentsActions::GetContinuousActionNum(int32& OutNum, const ULearni
 }
 
 bool ULearningAgentsActions::GetContinuousAction(
-	TArray<float>& OutValues,
-	const ULearningAgentsActionObject* Object,
-	const FLearningAgentsActionObjectElement Element,
+	TArray<float>& OutValues, 
+	const ULearningAgentsActionObject* Object, 
+	const FLearningAgentsActionObjectElement Element, 
 	const bool bActive,
 	const FName Tag,
-	const FLearningAgentsVisualLoggerSettings& VisualLoggerSettings)
+	const bool bVisualLoggerEnabled,
+	ULearningAgentsManagerListener* VisualLoggerListener,
+	const int32 VisualLoggerAgentId,
+	const FVector VisualLoggerLocation,
+	const FLinearColor VisualLoggerColor)
 {
 	if (!bActive)
 	{
@@ -3902,7 +3961,7 @@ bool ULearningAgentsActions::GetContinuousAction(
 
 	OutValues.SetNumUninitialized(OutValueNum);
 
-	if (!GetContinuousActionToArrayView(OutValues, Object, Element, bActive, Tag, VisualLoggerSettings))
+	if (!GetContinuousActionToArrayView(OutValues, Object, Element, bActive, Tag, bVisualLoggerEnabled, VisualLoggerListener, VisualLoggerAgentId, VisualLoggerLocation, VisualLoggerColor))
 	{
 		OutValues.Empty();
 		return false;
@@ -3912,12 +3971,16 @@ bool ULearningAgentsActions::GetContinuousAction(
 }
 
 bool ULearningAgentsActions::GetContinuousActionToArrayView(
-	TArrayView<float> OutValues,
-	const ULearningAgentsActionObject* Object,
-	const FLearningAgentsActionObjectElement Element,
+	TArrayView<float> OutValues, 
+	const ULearningAgentsActionObject* Object, 
+	const FLearningAgentsActionObjectElement Element, 
 	const bool bActive,
 	const FName Tag,
-	const FLearningAgentsVisualLoggerSettings& VisualLoggerSettings)
+	const bool bVisualLoggerEnabled,
+	ULearningAgentsManagerListener* VisualLoggerListener,
+	const int32 VisualLoggerAgentId,
+	const FVector VisualLoggerLocation,
+	const FLinearColor VisualLoggerColor)
 {
 	if (!bActive)
 	{
@@ -3967,16 +4030,16 @@ bool ULearningAgentsActions::GetContinuousActionToArrayView(
 	UE::Learning::Array::Copy<1, float>(OutValues, Values);
 
 #if UE_LEARNING_AGENTS_ENABLE_VISUAL_LOG
-	if (VisualLoggerSettings.bEnabled && VisualLoggerSettings.Listener)
+	if (bVisualLoggerEnabled && VisualLoggerListener)
 	{
-		const ULearningAgentsVisualLoggerObject* VisualLoggerObject = VisualLoggerSettings.Listener->GetOrAddVisualLoggerObject(Tag);
+		const ULearningAgentsVisualLoggerObject* VisualLoggerObject = VisualLoggerListener->GetOrAddVisualLoggerObject(Tag);
 
-		UE_LEARNING_AGENTS_VLOG_STRING(VisualLoggerObject, LogLearning, Display, VisualLoggerSettings.Location,
-			VisualLoggerSettings.Color.ToFColor(true),
+		UE_LEARNING_AGENTS_VLOG_STRING(VisualLoggerObject, LogLearning, Display, VisualLoggerLocation,
+			VisualLoggerColor.ToFColor(true),
 			TEXT("Listener: %s\nTag: %s\nAgent Id: % 3i\nEncoded: %s\nValues: %s"),
-			*VisualLoggerSettings.Listener->GetName(),
+			*VisualLoggerListener->GetName(),
 			*Tag.ToString(),
-			VisualLoggerSettings.AgentId,
+			VisualLoggerAgentId,
 			*UE::Learning::Array::FormatFloat(OutValues),
 			*UE::Learning::Array::FormatFloat(OutValues)); // Encoded is identical to provided values
 	}
@@ -3986,12 +4049,16 @@ bool ULearningAgentsActions::GetContinuousActionToArrayView(
 }
 
 bool ULearningAgentsActions::GetExclusiveDiscreteAction(
-	int32& OutIndex,
-	const ULearningAgentsActionObject* Object,
-	const FLearningAgentsActionObjectElement Element,
-	const bool bActive,
+	int32& OutIndex, 
+	const ULearningAgentsActionObject* Object, 
+	const FLearningAgentsActionObjectElement Element, 
+	const bool bActive,	
 	const FName Tag,
-	const FLearningAgentsVisualLoggerSettings& VisualLoggerSettings)
+	const bool bVisualLoggerEnabled,
+	ULearningAgentsManagerListener* VisualLoggerListener,
+	const int32 VisualLoggerAgentId,
+	const FVector VisualLoggerLocation,
+	const FLinearColor VisualLoggerColor)
 {
 	if (!bActive)
 	{
@@ -4030,16 +4097,16 @@ bool ULearningAgentsActions::GetExclusiveDiscreteAction(
 	OutIndex = Object->ActionObject.GetDiscreteExclusive(Element.ObjectElement).DiscreteIndex;
 
 #if UE_LEARNING_AGENTS_ENABLE_VISUAL_LOG
-	if (VisualLoggerSettings.bEnabled && VisualLoggerSettings.Listener)
+	if (bVisualLoggerEnabled && VisualLoggerListener)
 	{
-		const ULearningAgentsVisualLoggerObject* VisualLoggerObject = VisualLoggerSettings.Listener->GetOrAddVisualLoggerObject(Tag);
+		const ULearningAgentsVisualLoggerObject* VisualLoggerObject = VisualLoggerListener->GetOrAddVisualLoggerObject(Tag);
 
-		UE_LEARNING_AGENTS_VLOG_STRING(VisualLoggerObject, LogLearning, Display, VisualLoggerSettings.Location,
-			VisualLoggerSettings.Color.ToFColor(true),
+		UE_LEARNING_AGENTS_VLOG_STRING(VisualLoggerObject, LogLearning, Display, VisualLoggerLocation,
+			VisualLoggerColor.ToFColor(true),
 			TEXT("Listener: %s\nTag: %s\nAgent Id: % 3i\nIndex: [%i]"),
-			*VisualLoggerSettings.Listener->GetName(),
+			*VisualLoggerListener->GetName(),
 			*Tag.ToString(),
-			VisualLoggerSettings.AgentId,
+			VisualLoggerAgentId,
 			OutIndex);
 	}
 #endif
@@ -4053,7 +4120,11 @@ bool ULearningAgentsActions::GetNamedExclusiveDiscreteAction(
 	const FLearningAgentsActionObjectElement Element,
 	const bool bActive,
 	const FName Tag,
-	const FLearningAgentsVisualLoggerSettings& VisualLoggerSettings)
+	const bool bVisualLoggerEnabled,
+	ULearningAgentsManagerListener* VisualLoggerListener,
+	const int32 VisualLoggerAgentId,
+	const FVector VisualLoggerLocation,
+	const FLinearColor VisualLoggerColor)
 {
 	if (!bActive)
 	{
@@ -4092,16 +4163,16 @@ bool ULearningAgentsActions::GetNamedExclusiveDiscreteAction(
 	OutName = Object->ActionObject.GetNamedDiscreteExclusive(Element.ObjectElement).ElementName;
 
 #if UE_LEARNING_AGENTS_ENABLE_VISUAL_LOG
-	if (VisualLoggerSettings.bEnabled && VisualLoggerSettings.Listener)
+	if (bVisualLoggerEnabled && VisualLoggerListener)
 	{
-		const ULearningAgentsVisualLoggerObject* VisualLoggerObject = VisualLoggerSettings.Listener->GetOrAddVisualLoggerObject(Tag);
+		const ULearningAgentsVisualLoggerObject* VisualLoggerObject = VisualLoggerListener->GetOrAddVisualLoggerObject(Tag);
 
-		UE_LEARNING_AGENTS_VLOG_STRING(VisualLoggerObject, LogLearning, Display, VisualLoggerSettings.Location,
-			VisualLoggerSettings.Color.ToFColor(true),
+		UE_LEARNING_AGENTS_VLOG_STRING(VisualLoggerObject, LogLearning, Display, VisualLoggerLocation,
+			VisualLoggerColor.ToFColor(true),
 			TEXT("Listener: %s\nTag: %s\nAgent Id: % 3i\nName: [%s]"),
-			*VisualLoggerSettings.Listener->GetName(),
+			*VisualLoggerListener->GetName(),
 			*Tag.ToString(),
-			VisualLoggerSettings.AgentId,
+			VisualLoggerAgentId,
 			*OutName.ToString());
 	}
 #endif
@@ -4150,12 +4221,16 @@ bool ULearningAgentsActions::GetInclusiveDiscreteActionNum(int32& OutNum, const
 }
 
 bool ULearningAgentsActions::GetInclusiveDiscreteAction(
-	TArray<int32>& OutIndices,
-	const ULearningAgentsActionObject* Object,
-	const FLearningAgentsActionObjectElement Element,
+	TArray<int32>& OutIndices, 
+	const ULearningAgentsActionObject* Object, 
+	const FLearningAgentsActionObjectElement Element, 
 	const bool bActive,
 	const FName Tag,
-	const FLearningAgentsVisualLoggerSettings& VisualLoggerSettings)
+	const bool bVisualLoggerEnabled,
+	ULearningAgentsManagerListener* VisualLoggerListener,
+	const int32 VisualLoggerAgentId,
+	const FVector VisualLoggerLocation,
+	const FLinearColor VisualLoggerColor)
 {
 	if (!bActive)
 	{
@@ -4171,7 +4246,7 @@ bool ULearningAgentsActions::GetInclusiveDiscreteAction(
 
 	OutIndices.SetNumUninitialized(OutIndexNum);
 
-	if (!GetInclusiveDiscreteActionToArrayView(OutIndices, Object, Element, bActive, Tag, VisualLoggerSettings))
+	if (!GetInclusiveDiscreteActionToArrayView(OutIndices, Object, Element, bActive, Tag, bVisualLoggerEnabled, VisualLoggerListener, VisualLoggerAgentId, VisualLoggerLocation, VisualLoggerColor))
 	{
 		OutIndices.Empty();
 		return false;
@@ -4182,11 +4257,15 @@ bool ULearningAgentsActions::GetInclusiveDiscreteAction(
 
 bool ULearningAgentsActions::GetInclusiveDiscreteActionToArrayView(
 	TArrayView<int32> OutIndices,
-	const ULearningAgentsActionObject* Object,
-	const FLearningAgentsActionObjectElement Element,
+	const ULearningAgentsActionObject* Object, 
+	const FLearningAgentsActionObjectElement Element, 
 	const bool bActive,
 	const FName Tag,
-	const FLearningAgentsVisualLoggerSettings& VisualLoggerSettings)
+	const bool bVisualLoggerEnabled,
+	ULearningAgentsManagerListener* VisualLoggerListener,
+	const int32 VisualLoggerAgentId,
+	const FVector VisualLoggerLocation,
+	const FLinearColor VisualLoggerColor)
 {
 	if (!bActive)
 	{
@@ -4236,16 +4315,16 @@ bool ULearningAgentsActions::GetInclusiveDiscreteActionToArrayView(
 	UE::Learning::Array::Copy<1, int32>(OutIndices, Indices);
 
 #if UE_LEARNING_AGENTS_ENABLE_VISUAL_LOG
-	if (VisualLoggerSettings.bEnabled && VisualLoggerSettings.Listener)
+	if (bVisualLoggerEnabled && VisualLoggerListener)
 	{
-		const ULearningAgentsVisualLoggerObject* VisualLoggerObject = VisualLoggerSettings.Listener->GetOrAddVisualLoggerObject(Tag);
+		const ULearningAgentsVisualLoggerObject* VisualLoggerObject = VisualLoggerListener->GetOrAddVisualLoggerObject(Tag);
 
-		UE_LEARNING_AGENTS_VLOG_STRING(VisualLoggerObject, LogLearning, Display, VisualLoggerSettings.Location,
-			VisualLoggerSettings.Color.ToFColor(true),
+		UE_LEARNING_AGENTS_VLOG_STRING(VisualLoggerObject, LogLearning, Display, VisualLoggerLocation,
+			VisualLoggerColor.ToFColor(true),
 			TEXT("Listener: %s\nTag: %s\nAgent Id: % 3i\nIndices: %s"),
-			*VisualLoggerSettings.Listener->GetName(),
+			*VisualLoggerListener->GetName(),
 			*Tag.ToString(),
-			VisualLoggerSettings.AgentId,
+			VisualLoggerAgentId,
 			*UE::Learning::Array::FormatInt32(OutIndices, 256));
 	}
 #endif
@@ -4299,7 +4378,11 @@ bool ULearningAgentsActions::GetNamedInclusiveDiscreteAction(
 	const FLearningAgentsActionObjectElement Element,
 	const bool bActive,
 	const FName Tag,
-	const FLearningAgentsVisualLoggerSettings& VisualLoggerSettings)
+	const bool bVisualLoggerEnabled,
+	ULearningAgentsManagerListener* VisualLoggerListener,
+	const int32 VisualLoggerAgentId,
+	const FVector VisualLoggerLocation,
+	const FLinearColor VisualLoggerColor)
 {
 	if (!bActive)
 	{
@@ -4315,7 +4398,7 @@ bool ULearningAgentsActions::GetNamedInclusiveDiscreteAction(
 
 	OutNames.SetNumUninitialized(OutNameNum);
 
-	if (!GetNamedInclusiveDiscreteActionToArrayView(OutNames, Object, Element, bActive, Tag, VisualLoggerSettings))
+	if (!GetNamedInclusiveDiscreteActionToArrayView(OutNames, Object, Element, bActive, Tag, bVisualLoggerEnabled, VisualLoggerListener, VisualLoggerAgentId, VisualLoggerLocation, VisualLoggerColor))
 	{
 		OutNames.Empty();
 		return false;
@@ -4328,9 +4411,13 @@ bool ULearningAgentsActions::GetNamedInclusiveDiscreteActionToArrayView(
 	TArrayView<FName> OutNames,
 	const ULearningAgentsActionObject* Object,
 	const FLearningAgentsActionObjectElement Element,
-	const bool bActive,
+	const bool bActive,	
 	const FName Tag,
-	const FLearningAgentsVisualLoggerSettings& VisualLoggerSettings)
+	const bool bVisualLoggerEnabled,
+	ULearningAgentsManagerListener* VisualLoggerListener,
+	const int32 VisualLoggerAgentId,
+	const FVector VisualLoggerLocation,
+	const FLinearColor VisualLoggerColor)
 {
 	if (!bActive)
 	{
@@ -4380,16 +4467,16 @@ bool ULearningAgentsActions::GetNamedInclusiveDiscreteActionToArrayView(
 	UE::Learning::Array::Copy<1, FName>(OutNames, Names);
 
 #if UE_LEARNING_AGENTS_ENABLE_VISUAL_LOG
-	if (VisualLoggerSettings.bEnabled && VisualLoggerSettings.Listener)
+	if (bVisualLoggerEnabled && VisualLoggerListener)
 	{
-		const ULearningAgentsVisualLoggerObject* VisualLoggerObject = VisualLoggerSettings.Listener->GetOrAddVisualLoggerObject(Tag);
+		const ULearningAgentsVisualLoggerObject* VisualLoggerObject = VisualLoggerListener->GetOrAddVisualLoggerObject(Tag);
 
-		UE_LEARNING_AGENTS_VLOG_STRING(VisualLoggerObject, LogLearning, Display, VisualLoggerSettings.Location,
-			VisualLoggerSettings.Color.ToFColor(true),
+		UE_LEARNING_AGENTS_VLOG_STRING(VisualLoggerObject, LogLearning, Display, VisualLoggerLocation,
+			VisualLoggerColor.ToFColor(true),
 			TEXT("Listener: %s\nTag: %s\nAgent Id: % 3i\nNames: %s"),
-			*VisualLoggerSettings.Listener->GetName(),
+			*VisualLoggerListener->GetName(),
 			*Tag.ToString(),
-			VisualLoggerSettings.AgentId,
+			VisualLoggerAgentId,
 			*UE::Learning::Agents::Action::Private::FormatNames(OutNames, 256));
 	}
 #endif
@@ -4433,7 +4520,7 @@ bool ULearningAgentsActions::GetStructActionNum(int32& OutNum, const ULearningAg
 		OutNum = 0;
 		return false;
 	}
-
+	
 	const UE::Learning::Action::FObjectAndParameters Parameters = Object->ActionObject.GetAnd(Element.ObjectElement);
 
 	OutNum = Parameters.Elements.Num();
@@ -4937,7 +5024,7 @@ bool ULearningAgentsActions::GetStaticArrayActionToArrayView(TArrayView<FLearnin
 
 	if (Object->ActionObject.GetType(Element.ObjectElement) != UE::Learning::Action::EType::Array)
 	{
-		UE_LOG(LogLearning, Error, TEXT("GetStaticArrayActionToArrayView: Action '%s' type does not match. Action is '%s' but asked for '%s'."),
+		UE_LOG(LogLearning, Error, TEXT("GetStaticArrayActionToArrayView: Action '%s' type does not match. Action is '%s' but asked for '%s'."), 
 			*Object->ActionObject.GetTag(Element.ObjectElement).ToString(),
 			UE::Learning::Agents::Action::Private::GetActionTypeString(Object->ActionObject.GetType(Element.ObjectElement)),
 			UE::Learning::Agents::Action::Private::GetActionTypeString(UE::Learning::Action::EType::Array));
@@ -4998,13 +5085,17 @@ bool ULearningAgentsActions::GetPairAction(FLearningAgentsActionObjectElement& O
 }
 
 bool ULearningAgentsActions::GetEnumAction(
-	uint8& OutEnumValue,
-	const ULearningAgentsActionObject* Object,
-	const FLearningAgentsActionObjectElement Element,
-	const UEnum* Enum,
+	uint8& OutEnumValue, 
+	const ULearningAgentsActionObject* Object, 
+	const FLearningAgentsActionObjectElement Element, 
+	const UEnum* Enum, 
 	const bool bActive,
 	const FName Tag,
-	const FLearningAgentsVisualLoggerSettings& VisualLoggerSettings)
+	const bool bVisualLoggerEnabled,
+	ULearningAgentsManagerListener* VisualLoggerListener,
+	const int32 VisualLoggerAgentId,
+	const FVector VisualLoggerLocation,
+	const FLinearColor VisualLoggerColor)
 {
 	if (!bActive)
 	{
@@ -5044,16 +5135,16 @@ bool ULearningAgentsActions::GetEnumAction(
 	OutEnumValue = (uint8)EnumValue;
 
 #if UE_LEARNING_AGENTS_ENABLE_VISUAL_LOG
-	if (VisualLoggerSettings.bEnabled && VisualLoggerSettings.Listener)
+	if (bVisualLoggerEnabled && VisualLoggerListener)
 	{
-		const ULearningAgentsVisualLoggerObject* VisualLoggerObject = VisualLoggerSettings.Listener->GetOrAddVisualLoggerObject(Tag);
+		const ULearningAgentsVisualLoggerObject* VisualLoggerObject = VisualLoggerListener->GetOrAddVisualLoggerObject(Tag);
 
-		UE_LEARNING_AGENTS_VLOG_STRING(VisualLoggerObject, LogLearning, Display, VisualLoggerSettings.Location,
-			VisualLoggerSettings.Color.ToFColor(true),
+		UE_LEARNING_AGENTS_VLOG_STRING(VisualLoggerObject, LogLearning, Display, VisualLoggerLocation,
+			VisualLoggerColor.ToFColor(true),
 			TEXT("Listener: %s\nTag: %s\nAgent Id: % 3i\nEnum: %s\nSize: [%i]\nValue: [%s]\nIndex: [%i]"),
-			*VisualLoggerSettings.Listener->GetName(),
+			*VisualLoggerListener->GetName(),
 			*Tag.ToString(),
-			VisualLoggerSettings.AgentId,
+			VisualLoggerAgentId,
 			*Enum->GetName(),
 			Enum->NumEnums() - 1,
 			*Enum->GetDisplayNameTextByValue(OutEnumValue).ToString(),
@@ -5065,13 +5156,17 @@ bool ULearningAgentsActions::GetEnumAction(
 }
 
 bool ULearningAgentsActions::GetBitmaskAction(
-	int32& OutBitmaskValue,
-	const ULearningAgentsActionObject* Object,
-	const FLearningAgentsActionObjectElement Element,
-	const UEnum* Enum,
+	int32& OutBitmaskValue, 
+	const ULearningAgentsActionObject* Object, 
+	const FLearningAgentsActionObjectElement Element, 
+	const UEnum* Enum, 
 	const bool bActive,
 	const FName Tag,
-	const FLearningAgentsVisualLoggerSettings& VisualLoggerSettings)
+	const bool bVisualLoggerEnabled,
+	ULearningAgentsManagerListener* VisualLoggerListener,
+	const int32 VisualLoggerAgentId,
+	const FVector VisualLoggerLocation,
+	const FLinearColor VisualLoggerColor)
 {
 	if (!bActive)
 	{
@@ -5121,9 +5216,9 @@ bool ULearningAgentsActions::GetBitmaskAction(
 	}
 
 #if UE_LEARNING_AGENTS_ENABLE_VISUAL_LOG
-	if (VisualLoggerSettings.bEnabled && VisualLoggerSettings.Listener)
+	if (bVisualLoggerEnabled && VisualLoggerListener)
 	{
-		const ULearningAgentsVisualLoggerObject* VisualLoggerObject = VisualLoggerSettings.Listener->GetOrAddVisualLoggerObject(Tag);
+		const ULearningAgentsVisualLoggerObject* VisualLoggerObject = VisualLoggerListener->GetOrAddVisualLoggerObject(Tag);
 
 		FString ValuesString;
 		FString IndicesString;
@@ -5140,12 +5235,12 @@ bool ULearningAgentsActions::GetBitmaskAction(
 		ValuesString = ValuesString.TrimEnd();
 		IndicesString = IndicesString.TrimEnd();
 
-		UE_LEARNING_AGENTS_VLOG_STRING(VisualLoggerObject, LogLearning, Display, VisualLoggerSettings.Location,
-			VisualLoggerSettings.Color.ToFColor(true),
+		UE_LEARNING_AGENTS_VLOG_STRING(VisualLoggerObject, LogLearning, Display, VisualLoggerLocation,
+			VisualLoggerColor.ToFColor(true),
 			TEXT("Listener: %s\nTag: %s\nAgent Id: % 3i\nEnum: %s\nSize: [%i]\nValues: [%s]\nIndices: [%s]"),
-			*VisualLoggerSettings.Listener->GetName(),
+			*VisualLoggerListener->GetName(),
 			*Tag.ToString(),
-			VisualLoggerSettings.AgentId,
+			VisualLoggerAgentId,
 			*Enum->GetName(),
 			Enum->NumEnums() - 1,
 			*ValuesString,
@@ -5234,12 +5329,16 @@ bool ULearningAgentsActions::GetEncodingAction(FLearningAgentsActionObjectElemen
 }
 
 bool ULearningAgentsActions::GetBoolAction(
-	bool& bOutValue,
-	const ULearningAgentsActionObject* Object,
-	const FLearningAgentsActionObjectElement Element,
-	const bool bActive,
+	bool& bOutValue, 
+	const ULearningAgentsActionObject* Object, 
+	const FLearningAgentsActionObjectElement Element, 
+	const bool bActive,	
 	const FName Tag,
-	const FLearningAgentsVisualLoggerSettings& VisualLoggerSettings)
+	const bool bVisualLoggerEnabled,
+	ULearningAgentsManagerListener* VisualLoggerListener,
+	const int32 VisualLoggerAgentId,
+	const FVector VisualLoggerLocation,
+	const FLinearColor VisualLoggerColor)
 {
 	if (!bActive)
 	{
@@ -5254,32 +5353,36 @@ bool ULearningAgentsActions::GetBoolAction(
 	}
 
 	bOutValue = OutIndex == 1;
-
+	
 #if UE_LEARNING_AGENTS_ENABLE_VISUAL_LOG
-	if (VisualLoggerSettings.bEnabled && VisualLoggerSettings.Listener)
+	if (bVisualLoggerEnabled && VisualLoggerListener)
 	{
-		const ULearningAgentsVisualLoggerObject* VisualLoggerObject = VisualLoggerSettings.Listener->GetOrAddVisualLoggerObject(Tag);
+		const ULearningAgentsVisualLoggerObject* VisualLoggerObject = VisualLoggerListener->GetOrAddVisualLoggerObject(Tag);
 
-		UE_LEARNING_AGENTS_VLOG_STRING(VisualLoggerObject, LogLearning, Display, VisualLoggerSettings.Location,
-			VisualLoggerSettings.Color.ToFColor(true),
+		UE_LEARNING_AGENTS_VLOG_STRING(VisualLoggerObject, LogLearning, Display, VisualLoggerLocation,
+			VisualLoggerColor.ToFColor(true),
 			TEXT("Listener: %s\nTag: %s\nAgent Id: % 3i\nValue: [%s]"),
-			*VisualLoggerSettings.Listener->GetName(),
+			*VisualLoggerListener->GetName(),
 			*Tag.ToString(),
-			VisualLoggerSettings.AgentId,
+			VisualLoggerAgentId,
 			bOutValue ? TEXT("true") : TEXT("false"));
 	}
 #endif
-
+	
 	return true;
 }
 
 bool ULearningAgentsActions::GetFloatAction(
-	float& OutValue,
-	const ULearningAgentsActionObject* Object,
-	const FLearningAgentsActionObjectElement Element,
-	const bool bActive,
+	float& OutValue, 
+	const ULearningAgentsActionObject* Object, 
+	const FLearningAgentsActionObjectElement Element, 
+	const bool bActive, 
 	const FName Tag,
-	const FLearningAgentsVisualLoggerSettings& VisualLoggerSettings)
+	const bool bVisualLoggerEnabled,
+	ULearningAgentsManagerListener* VisualLoggerListener,
+	const int32 VisualLoggerAgentId,
+	const FVector VisualLoggerLocation,
+	const FLinearColor VisualLoggerColor)
 {
 	if (!bActive)
 	{
@@ -5293,16 +5396,16 @@ bool ULearningAgentsActions::GetFloatAction(
 	}
 
 #if UE_LEARNING_AGENTS_ENABLE_VISUAL_LOG
-	if (VisualLoggerSettings.bEnabled && VisualLoggerSettings.Listener)
+	if (bVisualLoggerEnabled && VisualLoggerListener)
 	{
-		const ULearningAgentsVisualLoggerObject* VisualLoggerObject = VisualLoggerSettings.Listener->GetOrAddVisualLoggerObject(Tag);
+		const ULearningAgentsVisualLoggerObject* VisualLoggerObject = VisualLoggerListener->GetOrAddVisualLoggerObject(Tag);
 
-		UE_LEARNING_AGENTS_VLOG_STRING(VisualLoggerObject, LogLearning, Display, VisualLoggerSettings.Location,
-			VisualLoggerSettings.Color.ToFColor(true),
+		UE_LEARNING_AGENTS_VLOG_STRING(VisualLoggerObject, LogLearning, Display, VisualLoggerLocation,
+			VisualLoggerColor.ToFColor(true),
 			TEXT("Listener: %s\nTag: %s\nAgent Id: % 3i\nValue: [% 6.1f]"),
-			*VisualLoggerSettings.Listener->GetName(),
+			*VisualLoggerListener->GetName(),
 			*Tag.ToString(),
-			VisualLoggerSettings.AgentId,
+			VisualLoggerAgentId,
 			OutValue);
 	}
 #endif
@@ -5311,13 +5414,17 @@ bool ULearningAgentsActions::GetFloatAction(
 }
 
 bool ULearningAgentsActions::GetLocationAction(
-	FVector& OutLocation,
-	const ULearningAgentsActionObject* Object,
-	const FLearningAgentsActionObjectElement Element,
-	const FTransform RelativeTransform,
+	FVector& OutLocation, 
+	const ULearningAgentsActionObject* Object, 
+	const FLearningAgentsActionObjectElement Element, 
+	const FTransform RelativeTransform, 
 	const bool bActive,
 	const FName Tag,
-	const FLearningAgentsVisualLoggerSettings& VisualLoggerSettings)
+	const bool bVisualLoggerEnabled,
+	ULearningAgentsManagerListener* VisualLoggerListener,
+	const int32 VisualLoggerAgentId,
+	const FVector VisualLoggerLocation,
+	const FLinearColor VisualLoggerColor)
 {
 	if (!bActive)
 	{
@@ -5333,52 +5440,57 @@ bool ULearningAgentsActions::GetLocationAction(
 
 	const FVector LocalLocation = FVector(OutValues[0], OutValues[1], OutValues[2]);
 	OutLocation = RelativeTransform.TransformPosition(LocalLocation);
-
+	
 #if UE_LEARNING_AGENTS_ENABLE_VISUAL_LOG
-	if (VisualLoggerSettings.bEnabled && VisualLoggerSettings.Listener)
+	if (bVisualLoggerEnabled && VisualLoggerListener)
 	{
-		const ULearningAgentsVisualLoggerObject* VisualLoggerObject = VisualLoggerSettings.Listener->GetOrAddVisualLoggerObject(Tag);
+		const ULearningAgentsVisualLoggerObject* VisualLoggerObject = VisualLoggerListener->GetOrAddVisualLoggerObject(Tag);
 
 		UE_LEARNING_AGENTS_VLOG_LOCATION(VisualLoggerObject, LogLearning, Display,
 			OutLocation,
 			10,
-			VisualLoggerSettings.Color.ToFColor(true),
+			VisualLoggerColor.ToFColor(true),
 			TEXT(""));
 
 		UE_LEARNING_AGENTS_VLOG_SEGMENT(VisualLoggerObject, LogLearning, Display,
 			RelativeTransform.GetTranslation(),
 			OutLocation,
-			VisualLoggerSettings.Color.ToFColor(true),
+			VisualLoggerColor.ToFColor(true),
 			TEXT(""));
 
 		UE_LEARNING_AGENTS_VLOG_TRANSFORM(VisualLoggerObject, LogLearning, Display,
 			RelativeTransform.GetTranslation(),
 			RelativeTransform.GetRotation(),
-			VisualLoggerSettings.Color.ToFColor(true),
+			VisualLoggerColor.ToFColor(true),
 			TEXT(""));
 
-		UE_LEARNING_AGENTS_VLOG_STRING(VisualLoggerObject, LogLearning, Display, VisualLoggerSettings.Location,
-			VisualLoggerSettings.Color.ToFColor(true),
+		UE_LEARNING_AGENTS_VLOG_STRING(VisualLoggerObject, LogLearning, Display, VisualLoggerLocation,
+			VisualLoggerColor.ToFColor(true),
 			TEXT("Listener: %s\nTag: %s\nAgent Id: % 3i\nLocal Location: [% 6.1f % 6.1f % 6.1f]\nLocation: [% 6.1f % 6.1f % 6.1f]"),
-			*VisualLoggerSettings.Listener->GetName(),
+			*VisualLoggerListener->GetName(),
 			*Tag.ToString(),
-			VisualLoggerSettings.AgentId,
+			VisualLoggerAgentId,
 			LocalLocation.X, LocalLocation.Y, LocalLocation.Z,
 			OutLocation.X, OutLocation.Y, OutLocation.Z);
 	}
 #endif
-
+	
 	return true;
 }
 
 bool ULearningAgentsActions::GetRotationAction(
-	FRotator& OutRotation,
-	const ULearningAgentsActionObject* Object,
-	const FLearningAgentsActionObjectElement Element,
-	const FRotator RelativeRotation,
+	FRotator& OutRotation, 
+	const ULearningAgentsActionObject* Object, 
+	const FLearningAgentsActionObjectElement Element, 
+	const FRotator RelativeRotation, 
 	const bool bActive,
 	const FName Tag,
-	const FLearningAgentsVisualLoggerSettings& VisualLoggerSettings)
+	const bool bVisualLoggerEnabled,
+	ULearningAgentsManagerListener* VisualLoggerListener,
+	const int32 VisualLoggerAgentId,
+	const FVector VisualLoggerRotationLocation,
+	const FVector VisualLoggerLocation,
+	const FLinearColor VisualLoggerColor)
 {
 	if (!bActive)
 	{
@@ -5386,7 +5498,7 @@ bool ULearningAgentsActions::GetRotationAction(
 	}
 
 	FQuat OutRotationQuat;
-	if (!GetRotationActionAsQuat(OutRotationQuat, Object, Element, FQuat::MakeFromRotator(RelativeRotation), bActive, Tag, VisualLoggerSettings))
+	if (!GetRotationActionAsQuat(OutRotationQuat, Object, Element, FQuat::MakeFromRotator(RelativeRotation), bActive, Tag, bVisualLoggerEnabled, VisualLoggerListener, VisualLoggerAgentId, VisualLoggerRotationLocation, VisualLoggerLocation, VisualLoggerColor))
 	{
 		OutRotation = FRotator::ZeroRotator;
 		return false;
@@ -5397,13 +5509,18 @@ bool ULearningAgentsActions::GetRotationAction(
 }
 
 bool ULearningAgentsActions::GetRotationActionAsQuat(
-	FQuat& OutRotation,
-	const ULearningAgentsActionObject* Object,
-	const FLearningAgentsActionObjectElement Element,
-	const FQuat RelativeRotation,
+	FQuat& OutRotation, 
+	const ULearningAgentsActionObject* Object, 
+	const FLearningAgentsActionObjectElement Element, 
+	const FQuat RelativeRotation, 
 	const bool bActive,
 	const FName Tag,
-	const FLearningAgentsVisualLoggerSettings& VisualLoggerSettings)
+	const bool bVisualLoggerEnabled,
+	ULearningAgentsManagerListener* VisualLoggerListener,
+	const int32 VisualLoggerAgentId,
+	const FVector VisualLoggerRotationLocation,
+	const FVector VisualLoggerLocation,
+	const FLinearColor VisualLoggerColor)
 {
 	if (!bActive)
 	{
@@ -5422,22 +5539,22 @@ bool ULearningAgentsActions::GetRotationActionAsQuat(
 	OutRotation = RelativeRotation * LocalRotation;
 
 #if UE_LEARNING_AGENTS_ENABLE_VISUAL_LOG
-	if (VisualLoggerSettings.bEnabled && VisualLoggerSettings.Listener)
+	if (bVisualLoggerEnabled && VisualLoggerListener)
 	{
-		const ULearningAgentsVisualLoggerObject* VisualLoggerObject = VisualLoggerSettings.Listener->GetOrAddVisualLoggerObject(Tag);
+		const ULearningAgentsVisualLoggerObject* VisualLoggerObject = VisualLoggerListener->GetOrAddVisualLoggerObject(Tag);
 
 		UE_LEARNING_AGENTS_VLOG_TRANSFORM(VisualLoggerObject, LogLearning, Display,
-			VisualLoggerSettings.DebugLocation,
+			VisualLoggerRotationLocation,
 			LocalRotation.Rotator(),
-			VisualLoggerSettings.Color.ToFColor(true),
+			VisualLoggerColor.ToFColor(true),
 			TEXT(""));
 
-		UE_LEARNING_AGENTS_VLOG_STRING(VisualLoggerObject, LogLearning, Display, VisualLoggerSettings.Location,
-			VisualLoggerSettings.Color.ToFColor(true),
+		UE_LEARNING_AGENTS_VLOG_STRING(VisualLoggerObject, LogLearning, Display, VisualLoggerLocation,
+			VisualLoggerColor.ToFColor(true),
 			TEXT("Listener: %s\nTag: %s\nAgent Id: % 3i\nLocal Rotation Vector: [% 6.1f % 6.1f % 6.1f]\nLocal Rotation: [% 6.1f % 6.1f % 6.1f % 6.1f]\nRotation: [% 6.1f % 6.1f % 6.1f % 6.1f]"),
-			*VisualLoggerSettings.Listener->GetName(),
+			*VisualLoggerListener->GetName(),
 			*Tag.ToString(),
-			VisualLoggerSettings.AgentId,
+			VisualLoggerAgentId,
 			LocalRotationVector.X, LocalRotationVector.Y, LocalRotationVector.Z,
 			LocalRotation.X, LocalRotation.Y, LocalRotation.Z, LocalRotation.W,
 			OutRotation.X, OutRotation.Y, OutRotation.Z, OutRotation.W);
@@ -5448,13 +5565,17 @@ bool ULearningAgentsActions::GetRotationActionAsQuat(
 }
 
 bool ULearningAgentsActions::GetScaleAction(
-	FVector& OutScale,
-	const ULearningAgentsActionObject* Object,
-	const FLearningAgentsActionObjectElement Element,
-	const FVector RelativeScale,
+	FVector& OutScale, 
+	const ULearningAgentsActionObject* Object, 
+	const FLearningAgentsActionObjectElement Element, 
+	const FVector RelativeScale, 
 	const bool bActive,
 	const FName Tag,
-	const FLearningAgentsVisualLoggerSettings& VisualLoggerSettings)
+	const bool bVisualLoggerEnabled,
+	ULearningAgentsManagerListener* VisualLoggerListener,
+	const int32 VisualLoggerAgentId,
+	const FVector VisualLoggerLocation,
+	const FLinearColor VisualLoggerColor)
 {
 	if (!bActive)
 	{
@@ -5472,16 +5593,16 @@ bool ULearningAgentsActions::GetScaleAction(
 	OutScale = RelativeScale * LocalScaleVector;
 
 #if UE_LEARNING_AGENTS_ENABLE_VISUAL_LOG
-	if (VisualLoggerSettings.bEnabled && VisualLoggerSettings.Listener)
+	if (bVisualLoggerEnabled && VisualLoggerListener)
 	{
-		const ULearningAgentsVisualLoggerObject* VisualLoggerObject = VisualLoggerSettings.Listener->GetOrAddVisualLoggerObject(Tag);
+		const ULearningAgentsVisualLoggerObject* VisualLoggerObject = VisualLoggerListener->GetOrAddVisualLoggerObject(Tag);
 
-		UE_LEARNING_AGENTS_VLOG_STRING(VisualLoggerObject, LogLearning, Display, VisualLoggerSettings.Location,
-			VisualLoggerSettings.Color.ToFColor(true),
+		UE_LEARNING_AGENTS_VLOG_STRING(VisualLoggerObject, LogLearning, Display, VisualLoggerLocation,
+			VisualLoggerColor.ToFColor(true),
 			TEXT("Listener: %s\nTag: %s\nAgent Id: % 3i\nLocal Scale: [% 6.1f % 6.1f % 6.1f]\nScale: [% 6.1f % 6.1f % 6.1f]"),
-			*VisualLoggerSettings.Listener->GetName(),
+			*VisualLoggerListener->GetName(),
 			*Tag.ToString(),
-			VisualLoggerSettings.AgentId,
+			VisualLoggerAgentId,
 			LocalScaleVector.X, LocalScaleVector.Y, LocalScaleVector.Z,
 			OutScale.X, OutScale.Y, OutScale.Z);
 	}
@@ -5491,13 +5612,17 @@ bool ULearningAgentsActions::GetScaleAction(
 }
 
 bool ULearningAgentsActions::GetTransformAction(
-	FTransform& OutTransform,
-	const ULearningAgentsActionObject* Object,
-	const FLearningAgentsActionObjectElement Element,
-	const FTransform RelativeTransform,
+	FTransform& OutTransform, 
+	const ULearningAgentsActionObject* Object, 
+	const FLearningAgentsActionObjectElement Element, 
+	const FTransform RelativeTransform, 
 	const bool bActive,
 	const FName Tag,
-	const FLearningAgentsVisualLoggerSettings& VisualLoggerSettings)
+	const bool bVisualLoggerEnabled,
+	ULearningAgentsManagerListener* VisualLoggerListener,
+	const int32 VisualLoggerAgentId,
+	const FVector VisualLoggerLocation,
+	const FLinearColor VisualLoggerColor)
 {
 	if (!bActive)
 	{
@@ -5536,22 +5661,22 @@ bool ULearningAgentsActions::GetTransformAction(
 	OutTransform = FTransform(OutRotation, OutLocation, OutScale);
 
 #if UE_LEARNING_AGENTS_ENABLE_VISUAL_LOG
-	if (VisualLoggerSettings.bEnabled && VisualLoggerSettings.Listener)
+	if (bVisualLoggerEnabled && VisualLoggerListener)
 	{
-		const ULearningAgentsVisualLoggerObject* VisualLoggerObject = VisualLoggerSettings.Listener->GetOrAddVisualLoggerObject(Tag);
+		const ULearningAgentsVisualLoggerObject* VisualLoggerObject = VisualLoggerListener->GetOrAddVisualLoggerObject(Tag);
 
 		UE_LEARNING_AGENTS_VLOG_TRANSFORM(VisualLoggerObject, LogLearning, Display,
 			OutLocation,
 			OutRotation,
-			VisualLoggerSettings.Color.ToFColor(true),
+			VisualLoggerColor.ToFColor(true),
 			TEXT(""));
 
-		UE_LEARNING_AGENTS_VLOG_STRING(VisualLoggerObject, LogLearning, Display, VisualLoggerSettings.Location,
-			VisualLoggerSettings.Color.ToFColor(true),
+		UE_LEARNING_AGENTS_VLOG_STRING(VisualLoggerObject, LogLearning, Display, VisualLoggerLocation,
+			VisualLoggerColor.ToFColor(true),
 			TEXT("Listener: %s\nTag: %s\nAgent Id: % 3i\nLocation: [% 6.1f % 6.1f % 6.1f]\nRotation: [% 6.1f % 6.1f % 6.1f % 6.1f]\nScale: [% 6.1f % 6.1f % 6.1f]"),
-			*VisualLoggerSettings.Listener->GetName(),
+			*VisualLoggerListener->GetName(),
 			*Tag.ToString(),
-			VisualLoggerSettings.AgentId,
+			VisualLoggerAgentId,
 			OutLocation.X, OutLocation.Y, OutLocation.Z,
 			OutRotation.X, OutRotation.Y, OutRotation.Z, OutRotation.W,
 			OutScale.X, OutScale.Y, OutScale.Z);
@@ -5562,13 +5687,18 @@ bool ULearningAgentsActions::GetTransformAction(
 }
 
 bool ULearningAgentsActions::GetAngleAction(
-	float& OutAngle,
-	const ULearningAgentsActionObject* Object,
-	const FLearningAgentsActionObjectElement Element,
-	const float RelativeAngle,
+	float& OutAngle, 
+	const ULearningAgentsActionObject* Object, 
+	const FLearningAgentsActionObjectElement Element, 
+	const float RelativeAngle, 
 	const bool bActive,
 	const FName Tag,
-	const FLearningAgentsVisualLoggerSettings& VisualLoggerSettings)
+	const bool bVisualLoggerEnabled,
+	ULearningAgentsManagerListener* VisualLoggerListener,
+	const int32 VisualLoggerAgentId,
+	const FVector VisualLoggerAngleLocation,
+	const FVector VisualLoggerLocation,
+	const FLinearColor VisualLoggerColor)
 {
 	if (!bActive)
 	{
@@ -5585,24 +5715,24 @@ bool ULearningAgentsActions::GetAngleAction(
 	OutAngle = RelativeAngle + LocalAngle;
 
 #if UE_LEARNING_AGENTS_ENABLE_VISUAL_LOG
-	if (VisualLoggerSettings.bEnabled && VisualLoggerSettings.Listener)
+	if (bVisualLoggerEnabled && VisualLoggerListener)
 	{
-		const ULearningAgentsVisualLoggerObject* VisualLoggerObject = VisualLoggerSettings.Listener->GetOrAddVisualLoggerObject(Tag);
+		const ULearningAgentsVisualLoggerObject* VisualLoggerObject = VisualLoggerListener->GetOrAddVisualLoggerObject(Tag);
 
 		UE_LEARNING_AGENTS_VLOG_ANGLE_DEGREES(VisualLoggerObject, LogLearning, Display,
 			OutAngle,
 			0.0f,
-			VisualLoggerSettings.DebugLocation,
+			VisualLoggerAngleLocation,
 			10.0f,
-			VisualLoggerSettings.Color.ToFColor(true),
+			VisualLoggerColor.ToFColor(true),
 			TEXT(""));
 
-		UE_LEARNING_AGENTS_VLOG_STRING(VisualLoggerObject, LogLearning, Display, VisualLoggerSettings.Location,
-			VisualLoggerSettings.Color.ToFColor(true),
+		UE_LEARNING_AGENTS_VLOG_STRING(VisualLoggerObject, LogLearning, Display, VisualLoggerLocation,
+			VisualLoggerColor.ToFColor(true),
 			TEXT("Listener: %s\nTag: %s\nAgent Id: % 3i\nLocal Angle: [% 6.1f]\nAngle: [% 6.1f]"),
-			*VisualLoggerSettings.Listener->GetName(),
+			*VisualLoggerListener->GetName(),
 			*Tag.ToString(),
-			VisualLoggerSettings.AgentId,
+			VisualLoggerAgentId,
 			LocalAngle,
 			OutAngle);
 	}
@@ -5612,20 +5742,25 @@ bool ULearningAgentsActions::GetAngleAction(
 }
 
 bool ULearningAgentsActions::GetAngleActionRadians(
-	float& OutAngle,
-	const ULearningAgentsActionObject* Object,
-	const FLearningAgentsActionObjectElement Element,
-	const float RelativeAngle,
-	const bool bActive,
+	float& OutAngle, 
+	const ULearningAgentsActionObject* Object, 
+	const FLearningAgentsActionObjectElement Element, 
+	const float RelativeAngle, 
+	const bool bActive,	
 	const FName Tag,
-	const FLearningAgentsVisualLoggerSettings& VisualLoggerSettings)
+	const bool bVisualLoggerEnabled,
+	ULearningAgentsManagerListener* VisualLoggerListener,
+	const int32 VisualLoggerAgentId,
+	const FVector VisualLoggerAngleLocation,
+	const FVector VisualLoggerLocation,
+	const FLinearColor VisualLoggerColor)
 {
 	if (!bActive)
 	{
 		return GetNullAction(Object, Element, Tag);
 	}
 
-	if (!GetAngleAction(OutAngle, Object, Element, FMath::RadiansToDegrees(RelativeAngle), bActive, Tag, VisualLoggerSettings))
+	if (!GetAngleAction(OutAngle, Object, Element, FMath::RadiansToDegrees(RelativeAngle), bActive, Tag, bVisualLoggerEnabled, VisualLoggerListener, VisualLoggerAgentId, VisualLoggerAngleLocation, VisualLoggerLocation, VisualLoggerColor))
 	{
 		OutAngle = 0.0f;
 		return false;
@@ -5636,13 +5771,18 @@ bool ULearningAgentsActions::GetAngleActionRadians(
 }
 
 bool ULearningAgentsActions::GetVelocityAction(
-	FVector& OutVelocity,
-	const ULearningAgentsActionObject* Object,
-	const FLearningAgentsActionObjectElement Element,
-	const FTransform RelativeTransform,
+	FVector& OutVelocity, 
+	const ULearningAgentsActionObject* Object, 
+	const FLearningAgentsActionObjectElement Element, 
+	const FTransform RelativeTransform, 
 	const bool bActive,
 	const FName Tag,
-	const FLearningAgentsVisualLoggerSettings& VisualLoggerSettings)
+	const bool bVisualLoggerEnabled,
+	ULearningAgentsManagerListener* VisualLoggerListener,
+	const int32 VisualLoggerAgentId,
+	const FVector VisualLoggerVelocityLocation,
+	const FVector VisualLoggerLocation,
+	const FLinearColor VisualLoggerColor)
 {
 	if (!bActive)
 	{
@@ -5660,28 +5800,28 @@ bool ULearningAgentsActions::GetVelocityAction(
 	OutVelocity = RelativeTransform.TransformVector(LocalVelocity);
 
 #if UE_LEARNING_AGENTS_ENABLE_VISUAL_LOG
-	if (VisualLoggerSettings.bEnabled && VisualLoggerSettings.Listener)
+	if (bVisualLoggerEnabled && VisualLoggerListener)
 	{
-		const ULearningAgentsVisualLoggerObject* VisualLoggerObject = VisualLoggerSettings.Listener->GetOrAddVisualLoggerObject(Tag);
+		const ULearningAgentsVisualLoggerObject* VisualLoggerObject = VisualLoggerListener->GetOrAddVisualLoggerObject(Tag);
 
 		UE_LEARNING_AGENTS_VLOG_ARROW(VisualLoggerObject, LogLearning, Display,
-			VisualLoggerSettings.DebugLocation,
-			VisualLoggerSettings.DebugLocation + OutVelocity,
-			VisualLoggerSettings.Color.ToFColor(true),
+			VisualLoggerVelocityLocation,
+			VisualLoggerVelocityLocation + OutVelocity,
+			VisualLoggerColor.ToFColor(true),
 			TEXT(""));
 
 		UE_LEARNING_AGENTS_VLOG_TRANSFORM(VisualLoggerObject, LogLearning, Display,
 			RelativeTransform.GetTranslation(),
 			RelativeTransform.GetRotation(),
-			VisualLoggerSettings.Color.ToFColor(true),
+			VisualLoggerColor.ToFColor(true),
 			TEXT(""));
 
-		UE_LEARNING_AGENTS_VLOG_STRING(VisualLoggerObject, LogLearning, Display, VisualLoggerSettings.Location,
-			VisualLoggerSettings.Color.ToFColor(true),
+		UE_LEARNING_AGENTS_VLOG_STRING(VisualLoggerObject, LogLearning, Display, VisualLoggerLocation,
+			VisualLoggerColor.ToFColor(true),
 			TEXT("Listener: %s\nTag: %s\nAgent Id: % 3i\nLocal Velocity: [% 6.1f % 6.1f % 6.1f]\nVelocity: [% 6.1f % 6.1f % 6.1f]"),
-			*VisualLoggerSettings.Listener->GetName(),
+			*VisualLoggerListener->GetName(),
 			*Tag.ToString(),
-			VisualLoggerSettings.AgentId,
+			VisualLoggerAgentId,
 			LocalVelocity.X, LocalVelocity.Y, LocalVelocity.Z,
 			OutVelocity.X, OutVelocity.Y, OutVelocity.Z);
 	}
@@ -5691,13 +5831,19 @@ bool ULearningAgentsActions::GetVelocityAction(
 }
 
 bool ULearningAgentsActions::GetDirectionAction(
-	FVector& OutDirection,
-	const ULearningAgentsActionObject* Object,
-	const FLearningAgentsActionObjectElement Element,
-	const FTransform RelativeTransform,
-	const bool bActive,
+	FVector& OutDirection, 
+	const ULearningAgentsActionObject* Object, 
+	const FLearningAgentsActionObjectElement Element, 
+	const FTransform RelativeTransform, 
+	const bool bActive,	
 	const FName Tag,
-	const FLearningAgentsVisualLoggerSettings& VisualLoggerSettings)
+	const bool bVisualLoggerEnabled,
+	ULearningAgentsManagerListener* VisualLoggerListener,
+	const int32 VisualLoggerAgentId,
+	const FVector VisualLoggerDirectionLocation,
+	const FVector VisualLoggerLocation,
+	const float VisualLoggerArrowLength,
+	const FLinearColor VisualLoggerColor)
 {
 	if (!bActive)
 	{
@@ -5715,28 +5861,28 @@ bool ULearningAgentsActions::GetDirectionAction(
 	OutDirection = RelativeTransform.TransformVectorNoScale(LocalDirection);
 
 #if UE_LEARNING_AGENTS_ENABLE_VISUAL_LOG
-	if (VisualLoggerSettings.bEnabled && VisualLoggerSettings.Listener)
+	if (bVisualLoggerEnabled && VisualLoggerListener)
 	{
-		const ULearningAgentsVisualLoggerObject* VisualLoggerObject = VisualLoggerSettings.Listener->GetOrAddVisualLoggerObject(Tag);
+		const ULearningAgentsVisualLoggerObject* VisualLoggerObject = VisualLoggerListener->GetOrAddVisualLoggerObject(Tag);
 
 		UE_LEARNING_AGENTS_VLOG_ARROW(VisualLoggerObject, LogLearning, Display,
-			VisualLoggerSettings.DebugLocation,
-			VisualLoggerSettings.DebugLocation + VisualLoggerSettings.ArrowLength * OutDirection,
-			VisualLoggerSettings.Color.ToFColor(true),
+			VisualLoggerDirectionLocation,
+			VisualLoggerDirectionLocation + VisualLoggerArrowLength * OutDirection,
+			VisualLoggerColor.ToFColor(true),
 			TEXT(""));
 
 		UE_LEARNING_AGENTS_VLOG_TRANSFORM(VisualLoggerObject, LogLearning, Display,
 			RelativeTransform.GetTranslation(),
 			RelativeTransform.GetRotation(),
-			VisualLoggerSettings.Color.ToFColor(true),
+			VisualLoggerColor.ToFColor(true),
 			TEXT(""));
 
-		UE_LEARNING_AGENTS_VLOG_STRING(VisualLoggerObject, LogLearning, Display, VisualLoggerSettings.Location,
-			VisualLoggerSettings.Color.ToFColor(true),
+		UE_LEARNING_AGENTS_VLOG_STRING(VisualLoggerObject, LogLearning, Display, VisualLoggerLocation,
+			VisualLoggerColor.ToFColor(true),
 			TEXT("Listener: %s\nTag: %s\nAgent Id: % 3i\nEncoded: [% 6.2f % 6.2f % 6.2f]\nLocal Direction: [% 6.1f % 6.1f % 6.1f]\nDirection: [% 6.1f % 6.1f % 6.1f]"),
-			*VisualLoggerSettings.Listener->GetName(),
+			*VisualLoggerListener->GetName(),
 			*Tag.ToString(),
-			VisualLoggerSettings.AgentId,
+			VisualLoggerAgentId,
 			OutValues[0], OutValues[1], OutValues[2],
 			LocalDirection.X, LocalDirection.Y, LocalDirection.Z,
 			OutDirection.X, OutDirection.Y, OutDirection.Z);
@@ -5744,4 +5890,4 @@ bool ULearningAgentsActions::GetDirectionAction(
 #endif
 
 	return true;
-}
\ No newline at end of file
+}
