[LearningAgents] [LACombat]

#changelist validated
#virtualized
#Synced-CL 51612678

diff --git a/Engine/Plugins/Experimental/LearningAgents/Source/LearningAgentsTraining/Private/LearningAgentsFlowMatchingTrainer.cpp b/Engine/Plugins/Experimental/LearningAgents/Source/LearningAgentsTraining/Private/LearningAgentsFlowMatchingTrainer.cpp
--- a/Engine/Plugins/Experimental/LearningAgents/Source/LearningAgentsTraining/Private/LearningAgentsFlowMatchingTrainer.cpp
+++ b/Engine/Plugins/Experimental/LearningAgents/Source/LearningAgentsTraining/Private/LearningAgentsFlowMatchingTrainer.cpp
@@ -32,7 +32,6 @@ TSharedRef<FJsonObject> FLearningAgentsFlowMatchingTrainerTrainingSettings::AsJs
 
 	ConfigObject->SetNumberField(TEXT("ObservationNoiseScale"), Hyperparameters.ObservationNoiseScale);
 	ConfigObject->SetBoolField(TEXT("NormalizeLatents"), Hyperparameters.bNormalizeLatents);
-	ConfigObject->SetBoolField(TEXT("EmbedNormalizationsInWeights"), Hyperparameters.bEmbedNormalizationsInWeights);
 	ConfigObject->SetNumberField(TEXT("Seed"), Hyperparameters.RandomSeed);
 
 	ConfigObject->SetBoolField(TEXT("SaveSnapshots"), LoggingSettings.bSaveSnapshots);
@@ -472,6 +471,8 @@ TSharedRef<FJsonObject> ULearningAgentsFlowMatchingTrainer::CreateTrainerConfig(
 	ConfigObject->SetObjectField(TEXT("FlowMatchingTrainingSettings"), TrainingSettings.AsJsonConfig());
 	ConfigObject->SetObjectField(TEXT("FlowMatchingSettings"), Model->AsJsonConfig());
 
+	ConfigObject->SetBoolField(TEXT("EmbedNormalizationsInWeights"), TrainingSettings.bEmbedNormalizationsInWeights);
+
 	ConfigObject->SetStringField(TEXT("Device"), UE::Learning::Trainer::GetDeviceString(UE::Learning::Agents::GetTrainingDevice(TrainingSettings.Device)));
 	ConfigObject->SetBoolField(TEXT("UseTensorBoard"), TrainingSettings.LoggingSettings.bUseTensorboard);
 	ConfigObject->SetBoolField(TEXT("UseMLflow"), TrainingSettings.LoggingSettings.bUseMLflow);
diff --git a/Engine/Plugins/Experimental/LearningAgents/Source/LearningAgentsTraining/Private/LearningAgentsImitationTrainer.cpp b/Engine/Plugins/Experimental/LearningAgents/Source/LearningAgentsTraining/Private/LearningAgentsImitationTrainer.cpp
--- a/Engine/Plugins/Experimental/LearningAgents/Source/LearningAgentsTraining/Private/LearningAgentsImitationTrainer.cpp
+++ b/Engine/Plugins/Experimental/LearningAgents/Source/LearningAgentsTraining/Private/LearningAgentsImitationTrainer.cpp
@@ -28,7 +28,6 @@ TSharedRef<FJsonObject> FLearningAgentsImitationTrainerTrainingSettings::AsJsonC
 	ConfigObject->SetNumberField(TEXT("LearningRate"), Hyperparameters.LearningRate);
 	ConfigObject->SetNumberField(TEXT("LearningRateDecay"), Hyperparameters.LearningRateDecay);
 	ConfigObject->SetNumberField(TEXT("LearningRateDecayStepSize"), Hyperparameters.LearningRateDecayStepSize);
-	ConfigObject->SetBoolField(TEXT("EmbedNormalizationsInWeights"), Hyperparameters.bEmbedNormalizationsInWeights);
 	ConfigObject->SetNumberField(TEXT("WeightDecay"), Hyperparameters.WeightDecay);
 	ConfigObject->SetNumberField(TEXT("BatchSize"), Hyperparameters.BatchSize);
 	ConfigObject->SetNumberField(TEXT("Window"), Hyperparameters.Window);
@@ -483,7 +482,9 @@ TSharedRef<FJsonObject> ULearningAgentsImitationTrainer::CreateTrainerConfig(con
 	// Add Imitation Specific Config Entries
 	ConfigObject->SetObjectField(TEXT("BehaviorCloningSettings"), TrainingSettings.AsJsonConfig());
 	ConfigObject->SetObjectField(TEXT("PolicySettings"), Policy->AsJsonConfig());
-	ConfigObject->SetNumberField(TEXT("MemoryStateNum"), Policy->GetMemoryStateSize());
+	
+	// Add Normalization Choice
+	ConfigObject->SetBoolField(TEXT("EmbedNormalizationsInWeights"), TrainingSettings.bEmbedNormalizationsInWeights);
 
 	// Add Training Workflow Config Entries
 	ConfigObject->SetStringField(TEXT("Device"), UE::Learning::Trainer::GetDeviceString(UE::Learning::Agents::GetTrainingDevice(TrainingSettings.Device)));
diff --git a/Engine/Plugins/Experimental/LearningAgents/Source/LearningAgentsTraining/Public/LearningAgentsImitationTrainer.h b/Engine/Plugins/Experimental/LearningAgents/Source/LearningAgentsTraining/Public/LearningAgentsImitationTrainer.h
--- a/Engine/Plugins/Experimental/LearningAgents/Source/LearningAgentsTraining/Public/LearningAgentsImitationTrainer.h
+++ b/Engine/Plugins/Experimental/LearningAgents/Source/LearningAgentsTraining/Public/LearningAgentsImitationTrainer.h
@@ -63,14 +63,6 @@ public:
 	UPROPERTY(EditAnywhere, Category = "LearningAgents", meta = (ClampMin = "1", UIMin = "1"))
 	int32 LearningRateDecayStepSize = 1000;
 
-	/**
-	 * Whether to embed dataset normalizations into encoder and decoder networks. 
-	 * This is recommended to set to true to stabilize trainings.
-	 * <Important> The normalization option in action/observation schemas must be set to one of the "auto" options.
-	 */
-	UPROPERTY(EditAnywhere, Category = "LearningAgents")
-	bool bEmbedNormalizationsInWeights = false;
-
 	/**
 	 * Amount of weight decay to apply to the network. Larger values encourage network weights to be smaller but too
 	 * large a value can cause the network weights to collapse to all zeros.
@@ -191,9 +183,17 @@ public:
 	UPROPERTY(EditAnywhere, BlueprintReadWrite, Category = "LearningAgents")
 	ELearningAgentsTrainingDevice Device = ELearningAgentsTrainingDevice::GPU;
 
+	/**
+	 * Whether to embed dataset normalizations into encoder and decoder networks. 
+	 * This is recommended to set to true to stabilize trainings.
+	 * <Important> The normalization option in action/observation schemas must be set to one of the "auto" options.
+	 */
+	UPROPERTY(EditAnywhere, Category = "LearningAgents")
+	bool bEmbedNormalizationsInWeights = false;
+
 	/** Should evaluation run during the training process. Currently not used in Python */
 	UPROPERTY(VisibleDefaultsOnly, Category = "LearningAgents")
-	bool bRunEvaluation = true;
+	bool bRunEvaluation = false;
 
 	/** How much data should be used for evaluation. Currently not used in Python */
 	UPROPERTY(VisibleDefaultsOnly, Category = "LearningAgents", meta = (ClampMin = "0.0", UIMin = "0.0", ClampMax = "1.0", UIMax = "1.0", EditCondition = "bRunEvaluation"))
diff --git a/Engine/Plugins/Experimental/LearningAgents/Source/LearningAgentsTrainingEditor/Private/LearningAgentsFlowMatchingTrainerEditor.cpp b/Engine/Plugins/Experimental/LearningAgents/Source/LearningAgentsTrainingEditor/Private/LearningAgentsFlowMatchingTrainerEditor.cpp
--- a/Engine/Plugins/Experimental/LearningAgents/Source/LearningAgentsTrainingEditor/Private/LearningAgentsFlowMatchingTrainerEditor.cpp
+++ b/Engine/Plugins/Experimental/LearningAgents/Source/LearningAgentsTrainingEditor/Private/LearningAgentsFlowMatchingTrainerEditor.cpp
@@ -12,7 +12,14 @@ ALearningAgentsFlowMatchingTrainerEditor::ALearningAgentsFlowMatchingTrainerEdit
 
 void ALearningAgentsFlowMatchingTrainerEditor::StartTraining()
 {
-	SetupTraining();
+	if (ULearningAgentsFlowMatchingTrainer* Trainer = Cast<ULearningAgentsFlowMatchingTrainer>(SetupTraining()))
+	{
+		LearningAgentsFlowMatchingTrainer = Trainer;
+	}
+	else
+	{
+		LearningAgentsFlowMatchingTrainer = nullptr;
+	}
 
 	if (!LearningAgentsFlowMatchingTrainer)
 	{
diff --git a/Engine/Plugins/Experimental/LearningAgents/Source/LearningAgentsTrainingEditor/Private/LearningAgentsImitationTrainerEditor.cpp b/Engine/Plugins/Experimental/LearningAgents/Source/LearningAgentsTrainingEditor/Private/LearningAgentsImitationTrainerEditor.cpp
--- a/Engine/Plugins/Experimental/LearningAgents/Source/LearningAgentsTrainingEditor/Private/LearningAgentsImitationTrainerEditor.cpp
+++ b/Engine/Plugins/Experimental/LearningAgents/Source/LearningAgentsTrainingEditor/Private/LearningAgentsImitationTrainerEditor.cpp
@@ -11,8 +11,15 @@ ALearningAgentsImitationTrainerEditor::ALearningAgentsImitationTrainerEditor() {
 
 void ALearningAgentsImitationTrainerEditor::StartTraining()
 {
-	SetupTraining();
-	
+	if (ULearningAgentsImitationTrainer* Trainer = Cast<ULearningAgentsImitationTrainer>(SetupTraining()))
+	{
+		LearningAgentsImitationTrainer = Trainer;
+	}
+	else
+	{
+		LearningAgentsImitationTrainer = nullptr;
+	}
+
 	if (!LearningAgentsImitationTrainer)
 	{
 		UE_LOG(LogLearning, Error, TEXT("%s: Training has failed. Invalid imitation trainer. Was the imitation trainer set with SetupTraining?"), *GetName());
diff --git a/Engine/Plugins/Experimental/LearningAgents/Source/LearningAgentsTrainingEditor/Public/LearningAgentsTrainerEditorBase.h b/Engine/Plugins/Experimental/LearningAgents/Source/LearningAgentsTrainingEditor/Public/LearningAgentsTrainerEditorBase.h
--- a/Engine/Plugins/Experimental/LearningAgents/Source/LearningAgentsTrainingEditor/Public/LearningAgentsTrainerEditorBase.h
+++ b/Engine/Plugins/Experimental/LearningAgents/Source/LearningAgentsTrainingEditor/Public/LearningAgentsTrainerEditorBase.h
@@ -22,8 +22,8 @@ public:
 	ALearningAgentsTrainerEditorBase();
 
 	/** Setup the imitation trainer with necessary components in blueprints. */
-	UFUNCTION(BlueprintImplementableEvent, Category = "LearningAgents")
-	void SetupTraining();
+	UFUNCTION(BlueprintImplementableEvent, Category = "LearningAgents", meta = (ReturnDisplayName = "Trainer"))
+	ULearningAgentsManagerListener* SetupTraining();
 
 	/** Start training. */
 	virtual void StartTraining() PURE_VIRTUAL(ALearningAgentsTrainerEditorBase::StartTraining);
