LearningAgents: Fix issues with shared memory closing and SpecifyTransformObservation

#Synced-CL 52021387

diff --git a/Engine/Plugins/Experimental/LearningAgents/Source/LearningAgents/Private/LearningAgentsObservations.cpp b/Engine/Plugins/Experimental/LearningAgents/Source/LearningAgents/Private/LearningAgentsObservations.cpp
--- a/Engine/Plugins/Experimental/LearningAgents/Source/LearningAgents/Private/LearningAgentsObservations.cpp
+++ b/Engine/Plugins/Experimental/LearningAgents/Source/LearningAgents/Private/LearningAgentsObservations.cpp
@@ -1591,7 +1591,7 @@ FLearningAgentsObservationSchemaElement ULearningAgentsObservations::SpecifyTran
 			TEXT("Scale")
 		},
 		{
-			SpecifyLocationObservation(Schema, NormalizationSettings, true, TEXT("Location")),
+			SpecifyLocationObservation(Schema, NormalizationSettings),
 			SpecifyRotationObservation(Schema),
 			SpecifyScaleObservation(Schema)
 		},
