[LearningAgents] Modular Export of Networks/Replay/TrainerConfig

#Synced-CL 51878233

diff --git a/Engine/Plugins/Experimental/LearningAgents/Source/LearningAgentsTraining/Private/LearningAgentsFlowMatchingTrainer.cpp b/Engine/Plugins/Experimental/LearningAgents/Source/LearningAgentsTraining/Private/LearningAgentsFlowMatchingTrainer.cpp
--- a/Engine/Plugins/Experimental/LearningAgents/Source/LearningAgentsTraining/Private/LearningAgentsFlowMatchingTrainer.cpp
+++ b/Engine/Plugins/Experimental/LearningAgents/Source/LearningAgentsTraining/Private/LearningAgentsFlowMatchingTrainer.cpp
@@ -172,157 +172,30 @@ void ULearningAgentsFlowMatchingTrainer::BeginTraining(
 		return;
 	}
 
-	// Get Number of Steps
-	int32 TotalEpisodeNum = 0;
-	int32 TotalStepNum = 0;
-	for (const FLearningAgentsRecord& Record : Recording->Records)
-	{
-		if (FlowMatchingTrainerSettings.RecordingFilterTag.IsValid() && !Record.Tags.HasTag(FlowMatchingTrainerSettings.RecordingFilterTag))
-		{
-			UE_LOGFMT(LogLearning, Log, "Skipping record {} because its missing tag {}.", *Record.Name, *FlowMatchingTrainerSettings.RecordingFilterTag.ToString());
-			continue;
-		}
-
-		if (!Recording->Schemas.Contains(Record.SchemaName))
-		{
-			UE_LOGFMT(LogLearning, Error, "Schema name {Name} not found in schema list. Is the schema added and are the names correct?", ("Name", Record.SchemaName.ToString()));
-			continue;
-		}
-
-		if (FlowMatchingTrainerTrainingSettings.ValidationSettings.bCheckObservationSchemaCompatibility and !UE::Learning::Trainer::IsObservationSchemaSubsetCompatible(Recording->Schemas[Record.SchemaName].ObservationSchemaJson, Interactor->GetObservationSchema()->ObservationSchema, Interactor->GetObservationSchemaElement().SchemaElement))
-		{
-			UE_LOGFMT(LogLearning, Warning, "{Name}: Record has incompatible observation schema.", ("Name", GetName()));
-			continue;
-		}
-
-		if (FlowMatchingTrainerTrainingSettings.ValidationSettings.bCheckActionSchemaCompatibility and !UE::Learning::Trainer::IsActionSchemaSubsetCompatible(Recording->Schemas[Record.SchemaName].ActionSchemaJson, Interactor->GetActionSchema()->ActionSchema, Interactor->GetActionSchemaElement().SchemaElement))
-		{
-			UE_LOGFMT(LogLearning, Warning, "{Name}: Record has incompatible action schema.", ("Name", GetName()));
-			continue;
-		}
-
-		TotalEpisodeNum++;
-		TotalStepNum += Record.StepNum;
-	}
+	// Build Replay Buffer from Recording
+	ReplayBuffer = Recording->BuildReplayBuffer(
+		Interactor,
+		FlowMatchingTrainerSettings.RecordingFilterTag,
+		FlowMatchingTrainerTrainingSettings.ValidationSettings.bCheckObservationSchemaCompatibility,
+		FlowMatchingTrainerTrainingSettings.ValidationSettings.bCheckActionSchemaCompatibility);
 
-	if (TotalStepNum == 0)
+	if (!ReplayBuffer)
 	{
 		UE_LOGFMT(LogLearning, Warning, "{Name}: Recording contains no valid training data.", ("Name", GetName()));
 		return;
 	}
 
-	// Sizes
-	const int32 ObservationNum = Interactor->GetObservationVectorSize();
-	const int32 ActionNum = Interactor->GetActionVectorSize();
-
-	// Copy into Flat Arrays
-	TLearningArray<1, int32> RecordedEpisodeStarts;
-	TLearningArray<1, int32> RecordedEpisodeLengths;
-	TLearningArray<2, float> RecordedObservations;
-	TLearningArray<2, float> RecordedActions;
-
-	RecordedEpisodeStarts.SetNumUninitialized({ TotalEpisodeNum });
-	RecordedEpisodeLengths.SetNumUninitialized({ TotalEpisodeNum });
-	RecordedObservations.SetNumUninitialized({ TotalStepNum, ObservationNum });
-	RecordedActions.SetNumUninitialized({ TotalStepNum, ActionNum });
-
-	int32 EpisodeIdx = 0;
-	int32 StepIdx = 0;
-	for (const FLearningAgentsRecord& Record : Recording->Records)
-	{
-		if (FlowMatchingTrainerSettings.RecordingFilterTag.IsValid() && !Record.Tags.HasTag(FlowMatchingTrainerSettings.RecordingFilterTag))
-		{
-			UE_LOGFMT(LogLearning, Log, "Skipping record {Record} because its missing tag {Tag}.", ("Record", Record.Name), ("Tag", FlowMatchingTrainerSettings.RecordingFilterTag.ToString()));
-			continue;
-		}
-
-		if (!Recording->Schemas.Contains(Record.SchemaName))
-		{
-			UE_LOGFMT(LogLearning, Error, "Schema name {Name} not found in schema list. Is the schema added and are the names correct?", ("Name", Record.SchemaName.ToString()));
-			continue;
-		}
-
-		bool bObservationSchemaSubsetCompatible = (!FlowMatchingTrainerTrainingSettings.ValidationSettings.bCheckObservationSchemaCompatibility) or UE::Learning::Trainer::IsObservationSchemaSubsetCompatible(Recording->Schemas[Record.SchemaName].ObservationSchemaJson, Interactor->GetObservationSchema()->ObservationSchema, Interactor->GetObservationSchemaElement().SchemaElement);
-		bool bActionSchemaSubsetCompatible = (!FlowMatchingTrainerTrainingSettings.ValidationSettings.bCheckActionSchemaCompatibility) or UE::Learning::Trainer::IsActionSchemaSubsetCompatible(Recording->Schemas[Record.SchemaName].ActionSchemaJson, Interactor->GetActionSchema()->ActionSchema, Interactor->GetActionSchemaElement().SchemaElement);
-
-		if (Record.ObservationDimNum == ObservationNum && Record.ActionDimNum == ActionNum)
-		{
-			TLearningArrayView<2, const float> ObservationsView = TLearningArrayView<2, const float>(Record.Blob->ObservationData.GetData(), { Record.StepNum, Record.ObservationDimNum });
-			TLearningArrayView<2, const float> ActionsView = TLearningArrayView<2, const float>(Record.Blob->ActionData.GetData(), { Record.StepNum, Record.ActionDimNum });
-
-			RecordedEpisodeStarts[EpisodeIdx] = StepIdx;
-			RecordedEpisodeLengths[EpisodeIdx] = Record.StepNum;
-			UE::Learning::Array::Copy(RecordedObservations.Slice(StepIdx, Record.StepNum), ObservationsView);
-			UE::Learning::Array::Copy(RecordedActions.Slice(StepIdx, Record.StepNum), ActionsView);
-		}
-		else if (bObservationSchemaSubsetCompatible && bActionSchemaSubsetCompatible)
-		{
-			TLearningArrayView<2, const float> ObservationsView = TLearningArrayView<2, const float>(Record.Blob->ObservationData.GetData(), { Record.StepNum, Record.ObservationDimNum });
-			TLearningArrayView<2, const float> ActionsView = TLearningArrayView<2, const float>(Record.Blob->ActionData.GetData(), { Record.StepNum, Record.ActionDimNum });
-
-			RecordedEpisodeStarts[EpisodeIdx] = StepIdx;
-			RecordedEpisodeLengths[EpisodeIdx] = Record.StepNum;
-
-			TArray<int32> ObservationIndices = UE::Learning::Trainer::ComputeObservationSchemaSubsetIndices(
-				Recording->Schemas[Record.SchemaName].ObservationSchemaJson, Interactor->GetObservationSchema()->ObservationSchema, Interactor->GetObservationSchemaElement().SchemaElement);
-			check(ObservationIndices.Num() == ObservationNum);
-
-			for (int Step = 0; Step < Record.StepNum; Step++)
-			{
-				for (int i = 0; i < ObservationIndices.Num(); i++)
-				{
-					RecordedObservations[StepIdx + Step][i] = ObservationsView[Step][ObservationIndices[i]];
-				}
-			}
-
-			TArray<int32> ActionIndices = UE::Learning::Trainer::ComputeActionSchemaSubsetIndices(
-				Recording->Schemas[Record.SchemaName].ActionSchemaJson, Interactor->GetActionSchema()->ActionSchema, Interactor->GetActionSchemaElement().SchemaElement);
-			check(ActionIndices.Num() == ActionNum);
-
-			for (int Step = 0; Step < Record.StepNum; Step++)
-			{
-				for (int i = 0; i < ActionIndices.Num(); i++)
-				{
-					RecordedActions[StepIdx + Step][i] = ActionsView[Step][ActionIndices[i]];
-				}
-			}
-		}
-		else { continue; }
-
-		EpisodeIdx++;
-		StepIdx += Record.StepNum;
-	}
-
-	check(EpisodeIdx == TotalEpisodeNum);
-	check(StepIdx == TotalStepNum);
-
-	const int32 ObservationSchemaId = 0;
-	const int32 ActionSchemaId = 0;
-
-	// Create Replay Buffer from Records
-	ReplayBuffer = MakeUnique<UE::Learning::FReplayBuffer>();
-	ReplayBuffer->AddRecords(
-		TotalEpisodeNum,
-		TotalStepNum,
-		ObservationSchemaId,
-		ObservationNum,
-		ActionSchemaId,
-		ActionNum,
-		RecordedEpisodeStarts,
-		RecordedEpisodeLengths,
-		RecordedObservations,
-		RecordedActions);
-
 	DenoiserNetworkId = Trainer->AddNetwork(*Model->GetDenoiserNetworkAsset()->NeuralNetworkData);
 	EncoderNetworkId = Trainer->AddNetwork(*Model->GetEncoderNetworkAsset()->NeuralNetworkData);
 	DecoderNetworkId = Trainer->AddNetwork(*Model->GetDecoderNetworkAsset()->NeuralNetworkData);
 	ReplayBufferId = Trainer->AddReplayBuffer(*ReplayBuffer);
 
 	TSharedRef<FJsonObject> DataConfigObject = CreateDataConfig();
+	TSharedRef<FJsonObject> ModelConfigObject = CreateModelConfig();
 	TSharedRef<FJsonObject> TrainerConfigObject = CreateTrainerConfig(FlowMatchingTrainerTrainingSettings);
 
 	UE_LOGFMT(LogLearning, Display, "{Name}: Sending configs...", ("Name", GetName()));
-	SendConfigs(DataConfigObject, TrainerConfigObject);
+	SendConfigs(DataConfigObject, ModelConfigObject, TrainerConfigObject);
 
 	UE_LOGFMT(LogLearning, Display, "{Name}: Flow Matching Training Started", ("Name", GetName()));
 	UE::Learning::ETrainerResponse Response = UE::Learning::ETrainerResponse::Success;
@@ -368,6 +241,92 @@ void ULearningAgentsFlowMatchingTrainer::BeginTraining(
 	bIsTraining = true;
 }
 
+void ULearningAgentsFlowMatchingTrainer::ExportToFile(
+	UE::Learning::FFileTrainer& FileTrainer,
+	const ULearningAgentsRecording* Recording,
+	const FLearningAgentsFlowMatchingTrainerSettings& FlowMatchingTrainerSettings,
+	const FLearningAgentsFlowMatchingTrainerTrainingSettings& FlowMatchingTrainerTrainingSettings,
+	UE::Learning::ETrainerExportFlags ExportFlags)
+{
+	using namespace UE::Learning;
+
+	if (!IsSetup())
+	{
+		UE_LOGFMT(LogLearning, Error, "{Name}: Setup not complete.", ("Name", GetName()));
+		return;
+	}
+
+	if (IsTraining())
+	{
+		UE_LOGFMT(LogLearning, Error, "{Name}: Cannot export while training is active.", ("Name", GetName()));
+		return;
+	}
+
+	const bool bNeedReplayBuffers = EnumHasAnyFlags(ExportFlags, ETrainerExportFlags::ReplayBuffers);
+	const bool bNeedNetworks = EnumHasAnyFlags(ExportFlags, ETrainerExportFlags::Networks);
+	const bool bNeedTrainerConfig = EnumHasAnyFlags(ExportFlags, ETrainerExportFlags::TrainerConfig);
+
+	// Networks: register to get IDs (cheap)
+	if (bNeedNetworks)
+	{
+		DenoiserNetworkId = FileTrainer.AddNetwork(*Model->GetDenoiserNetworkAsset()->NeuralNetworkData);
+		EncoderNetworkId = FileTrainer.AddNetwork(*Model->GetEncoderNetworkAsset()->NeuralNetworkData);
+		DecoderNetworkId = FileTrainer.AddNetwork(*Model->GetDecoderNetworkAsset()->NeuralNetworkData);
+	}
+
+	// Replay buffers: build from recording (expensive)
+	if (bNeedReplayBuffers)
+	{
+		if (!Recording)
+		{
+			UE_LOGFMT(LogLearning, Error, "{Name}: Recording is nullptr.", ("Name", GetName()));
+			return;
+		}
+
+		if (Recording->Records.IsEmpty())
+		{
+			UE_LOGFMT(LogLearning, Error, "{Name}: Recording is empty!", ("Name", GetName()));
+			return;
+		}
+
+		ReplayBuffer = Recording->BuildReplayBuffer(
+			Interactor,
+			FlowMatchingTrainerSettings.RecordingFilterTag,
+			FlowMatchingTrainerTrainingSettings.ValidationSettings.bCheckObservationSchemaCompatibility,
+			FlowMatchingTrainerTrainingSettings.ValidationSettings.bCheckActionSchemaCompatibility);
+
+		if (!ReplayBuffer)
+		{
+			UE_LOGFMT(LogLearning, Warning, "{Name}: Recording contains no valid training data.", ("Name", GetName()));
+			return;
+		}
+
+		ReplayBufferId = FileTrainer.AddReplayBuffer(*ReplayBuffer);
+	}
+
+	// Build and send configs
+	TSharedRef<FJsonObject> DataConfigObject = bNeedReplayBuffers ? CreateDataConfig() : MakeShared<FJsonObject>();
+	TSharedRef<FJsonObject> ModelConfigObject = bNeedNetworks ? CreateModelConfig() : MakeShared<FJsonObject>();
+	TSharedRef<FJsonObject> TrainerConfigObject = bNeedTrainerConfig ? CreateTrainerConfig(FlowMatchingTrainerTrainingSettings) : MakeShared<FJsonObject>();
+	FileTrainer.SendConfigs(DataConfigObject, ModelConfigObject, TrainerConfigObject);
+
+	// Send network weights
+	if (bNeedNetworks)
+	{
+		FileTrainer.SendNetwork(DenoiserNetworkId, *Model->GetDenoiserNetworkAsset()->NeuralNetworkData);
+		FileTrainer.SendNetwork(EncoderNetworkId, *Model->GetEncoderNetworkAsset()->NeuralNetworkData);
+		FileTrainer.SendNetwork(DecoderNetworkId, *Model->GetDecoderNetworkAsset()->NeuralNetworkData);
+	}
+
+	// Send replay buffer data
+	if (bNeedReplayBuffers)
+	{
+		FileTrainer.SendReplayBuffer(ReplayBufferId, *ReplayBuffer);
+	}
+
+	UE_LOGFMT(LogLearning, Display, "{Name}: Export complete.", ("Name", GetName()));
+}
+
 TSharedRef<FJsonObject> ULearningAgentsFlowMatchingTrainer::CreateSchemaConfig(const int32 ObservationSchemaId, const int32 ActionSchemaId) const
 {
 	TSharedRef<FJsonObject> SchemasObject = MakeShared<FJsonObject>();
@@ -410,6 +369,24 @@ TSharedRef<FJsonObject> ULearningAgentsFlowMatchingTrainer::CreateDataConfig() c
 	const int32 ObservationSchemaId = 0;
 	const int32 ActionSchemaId = 0;
 
+	// Add Replay Buffers Config
+	TArray<TSharedPtr<FJsonValue>> ReplayBufferObjects;
+	TSharedRef<FJsonValueObject> ReplayBufferJsonValue = MakeShared<FJsonValueObject>(ReplayBuffer->AsJsonConfig(ReplayBufferId));
+	ReplayBufferObjects.Add(ReplayBufferJsonValue);
+	ConfigObject->SetArrayField(TEXT("ReplayBuffers"), ReplayBufferObjects);
+
+	ConfigObject->SetObjectField(TEXT("Schemas"), CreateSchemaConfig(ObservationSchemaId, ActionSchemaId));
+
+	return ConfigObject;
+}
+
+TSharedRef<FJsonObject> ULearningAgentsFlowMatchingTrainer::CreateModelConfig() const
+{
+	TSharedRef<FJsonObject> ConfigObject = MakeShared<FJsonObject>();
+
+	const int32 ObservationSchemaId = 0;
+	const int32 ActionSchemaId = 0;
+
 	// Add Neural Network Config Entries
 	TArray<TSharedPtr<FJsonValue>> NetworkObjects;
 
@@ -450,13 +427,7 @@ TSharedRef<FJsonObject> ULearningAgentsFlowMatchingTrainer::CreateDataConfig() c
 
 	ConfigObject->SetArrayField(TEXT("Networks"), NetworkObjects);
 
-	// Add Replay Buffers Config
-	TArray<TSharedPtr<FJsonValue>> ReplayBufferObjects;
-	TSharedRef<FJsonValueObject> ReplayBufferJsonValue = MakeShared<FJsonValueObject>(ReplayBuffer->AsJsonConfig(ReplayBufferId));
-	ReplayBufferObjects.Add(ReplayBufferJsonValue);
-	ConfigObject->SetArrayField(TEXT("ReplayBuffers"), ReplayBufferObjects);
-
-	ConfigObject->SetObjectField(TEXT("Schemas"), CreateSchemaConfig(ObservationSchemaId, ActionSchemaId));
+	ConfigObject->SetObjectField(TEXT("FlowMatchingSettings"), Model->AsJsonConfig());
 
 	return ConfigObject;
 }
@@ -469,7 +440,6 @@ TSharedRef<FJsonObject> ULearningAgentsFlowMatchingTrainer::CreateTrainerConfig(
 	ConfigObject->SetStringField(TEXT("TimeStamp"), *FDateTime::Now().ToFormattedString(TEXT("%Y-%m-%d_%H-%M-%s")));
 
 	ConfigObject->SetObjectField(TEXT("FlowMatchingTrainingSettings"), TrainingSettings.AsJsonConfig());
-	ConfigObject->SetObjectField(TEXT("FlowMatchingSettings"), Model->AsJsonConfig());
 
 	ConfigObject->SetBoolField(TEXT("EmbedNormalizationsInWeights"), TrainingSettings.bEmbedNormalizationsInWeights);
 
@@ -481,10 +451,10 @@ TSharedRef<FJsonObject> ULearningAgentsFlowMatchingTrainer::CreateTrainerConfig(
 	return ConfigObject;
 }
 
-void ULearningAgentsFlowMatchingTrainer::SendConfigs(const TSharedRef<FJsonObject>& DataConfigObject, const TSharedRef<FJsonObject>& TrainerConfigObject)
+void ULearningAgentsFlowMatchingTrainer::SendConfigs(const TSharedRef<FJsonObject>& DataConfigObject, const TSharedRef<FJsonObject>& ModelConfigObject, const TSharedRef<FJsonObject>& TrainerConfigObject)
 {
 	UE::Learning::ETrainerResponse Response = UE::Learning::ETrainerResponse::Success;
-	Response = Trainer->SendConfigs(DataConfigObject, TrainerConfigObject);
+	Response = Trainer->SendConfigs(DataConfigObject, ModelConfigObject, TrainerConfigObject);
 
 	if (Response != UE::Learning::ETrainerResponse::Success)
 	{
diff --git a/Engine/Plugins/Experimental/LearningAgents/Source/LearningAgentsTraining/Private/LearningAgentsImitationTrainer.cpp b/Engine/Plugins/Experimental/LearningAgents/Source/LearningAgentsTraining/Private/LearningAgentsImitationTrainer.cpp
--- a/Engine/Plugins/Experimental/LearningAgents/Source/LearningAgentsTraining/Private/LearningAgentsImitationTrainer.cpp
+++ b/Engine/Plugins/Experimental/LearningAgents/Source/LearningAgentsTraining/Private/LearningAgentsImitationTrainer.cpp
@@ -178,150 +178,19 @@ void ULearningAgentsImitationTrainer::BeginTraining(
 		return;
 	}
 
-	// Get Number of Steps
+	// Build Replay Buffer from Recording
+	ReplayBuffer = Recording->BuildReplayBuffer(
+		Interactor,
+		ImitationTrainerSettings.RecordingFilterTag,
+		ImitationTrainerTrainingSettings.ValidationSettings.bCheckObservationSchemaCompatibility,
+		ImitationTrainerTrainingSettings.ValidationSettings.bCheckActionSchemaCompatibility);
 
-	int32 TotalEpisodeNum = 0;
-	int32 TotalStepNum = 0;
-	for (const FLearningAgentsRecord& Record : Recording->Records)
-	{
-		if (ImitationTrainerSettings.RecordingFilterTag.IsValid() && !Record.Tags.HasTag(ImitationTrainerSettings.RecordingFilterTag))
-		{
-			UE_LOG(LogLearning, Log, TEXT("Skipping record %s because its missing tag %s."), *Record.Name, *ImitationTrainerSettings.RecordingFilterTag.ToString());
-			continue;
-		}
-
-		if (!Recording->Schemas.Contains(Record.SchemaName))
-		{
-			UE_LOG(LogLearning, Error, TEXT("Schema name %s not found in schema list. Is the schema added and are the names correct?"), *Record.SchemaName.ToString());
-			continue;
-		}
-
-		if (ImitationTrainerTrainingSettings.ValidationSettings.bCheckObservationSchemaCompatibility && !UE::Learning::Trainer::IsObservationSchemaSubsetCompatible(Recording->Schemas[Record.SchemaName].ObservationSchemaJson, Interactor->GetObservationSchema()->ObservationSchema, Interactor->GetObservationSchemaElement().SchemaElement))
-		{
-			UE_LOG(LogLearning, Warning, TEXT("%s: Record has incompatible observation schema."), *GetName());
-			continue;
-		}
-
-		if (ImitationTrainerTrainingSettings.ValidationSettings.bCheckActionSchemaCompatibility && !UE::Learning::Trainer::IsActionSchemaSubsetCompatible(Recording->Schemas[Record.SchemaName].ActionSchemaJson, Interactor->GetActionSchema()->ActionSchema, Interactor->GetActionSchemaElement().SchemaElement))
-		{
-			UE_LOG(LogLearning, Warning, TEXT("%s: Record has incompatible action schema."), *GetName());
-			continue;
-		}
-
-		TotalEpisodeNum++;
-		TotalStepNum += Record.StepNum;
-	}
-
-	if (TotalStepNum == 0)
+	if (!ReplayBuffer)
 	{
 		UE_LOG(LogLearning, Warning, TEXT("%s: Recording contains no valid training data."), *GetName());
 		return;
 	}
 
-	// Sizes
-
-	const int32 ObservationNum = Interactor->GetObservationVectorSize();
-	const int32 ActionNum = Interactor->GetActionVectorSize();
-
-	// Copy into Flat Arrays
-
-	TLearningArray<1, int32> RecordedEpisodeStarts;
-	TLearningArray<1, int32> RecordedEpisodeLengths;
-	TLearningArray<2, float> RecordedObservations;
-	TLearningArray<2, float> RecordedActions;
-
-	RecordedEpisodeStarts.SetNumUninitialized({ TotalEpisodeNum });
-	RecordedEpisodeLengths.SetNumUninitialized({ TotalEpisodeNum });
-	RecordedObservations.SetNumUninitialized({ TotalStepNum, ObservationNum });
-	RecordedActions.SetNumUninitialized({ TotalStepNum, ActionNum });
-
-	int32 EpisodeIdx = 0;
-	int32 StepIdx = 0;
-	for (const FLearningAgentsRecord& Record : Recording->Records)
-	{
-		if (ImitationTrainerSettings.RecordingFilterTag.IsValid() && !Record.Tags.HasTag(ImitationTrainerSettings.RecordingFilterTag))
-		{
-			UE_LOG(LogLearning, Log, TEXT("Skipping record %s because its missing tag %s."), *Record.Name, *ImitationTrainerSettings.RecordingFilterTag.ToString());
-			continue;
-		}
-
-		if (!Recording->Schemas.Contains(Record.SchemaName))
-		{
-			UE_LOG(LogLearning, Error, TEXT("Schema name %s not found in schema list. Is the schema added and are the names correct?"), *Record.SchemaName.ToString());
-			continue;
-		}
-
-		bool bObservationSchemaSubsetCompatible = (!ImitationTrainerTrainingSettings.ValidationSettings.bCheckObservationSchemaCompatibility) || UE::Learning::Trainer::IsObservationSchemaSubsetCompatible(Recording->Schemas[Record.SchemaName].ObservationSchemaJson, Interactor->GetObservationSchema()->ObservationSchema, Interactor->GetObservationSchemaElement().SchemaElement);
-		bool bActionSchemaSubsetCompatible = (!ImitationTrainerTrainingSettings.ValidationSettings.bCheckActionSchemaCompatibility) || UE::Learning::Trainer::IsActionSchemaSubsetCompatible(Recording->Schemas[Record.SchemaName].ActionSchemaJson, Interactor->GetActionSchema()->ActionSchema, Interactor->GetActionSchemaElement().SchemaElement);
-
-		if (Record.ObservationDimNum == ObservationNum && Record.ActionDimNum == ActionNum)
-		{
-			TLearningArrayView<2, const float> ObservationsView = TLearningArrayView<2, const float>(Record.Blob->ObservationData.GetData(), { Record.StepNum, Record.ObservationDimNum });
-			TLearningArrayView<2, const float> ActionsView = TLearningArrayView<2, const float>(Record.Blob->ActionData.GetData(), { Record.StepNum, Record.ActionDimNum });
-
-			RecordedEpisodeStarts[EpisodeIdx] = StepIdx;
-			RecordedEpisodeLengths[EpisodeIdx] = Record.StepNum;
-			UE::Learning::Array::Copy(RecordedObservations.Slice(StepIdx, Record.StepNum), ObservationsView);
-			UE::Learning::Array::Copy(RecordedActions.Slice(StepIdx, Record.StepNum), ActionsView);
-		}
-		else if (bObservationSchemaSubsetCompatible && bActionSchemaSubsetCompatible)
-		{
-			TLearningArrayView<2, const float> ObservationsView = TLearningArrayView<2, const float>(Record.Blob->ObservationData.GetData(), { Record.StepNum, Record.ObservationDimNum });
-			TLearningArrayView<2, const float> ActionsView = TLearningArrayView<2, const float>(Record.Blob->ActionData.GetData(), { Record.StepNum, Record.ActionDimNum });
-
-			RecordedEpisodeStarts[EpisodeIdx] = StepIdx;
-			RecordedEpisodeLengths[EpisodeIdx] = Record.StepNum;
-
-			TArray<int32> ObservationIndices = UE::Learning::Trainer::ComputeObservationSchemaSubsetIndices(
-				Recording->Schemas[Record.SchemaName].ObservationSchemaJson, Interactor->GetObservationSchema()->ObservationSchema, Interactor->GetObservationSchemaElement().SchemaElement);
-			check(ObservationIndices.Num() == ObservationNum);
-
-			for (int Step = 0; Step < Record.StepNum; Step++)
-			{
-				for (int i = 0; i < ObservationIndices.Num(); i++)
-				{
-					RecordedObservations[StepIdx + Step][i] = ObservationsView[Step][ObservationIndices[i]];
-				}
-			}
-
-			TArray<int32> ActionIndices = UE::Learning::Trainer::ComputeActionSchemaSubsetIndices(
-				Recording->Schemas[Record.SchemaName].ActionSchemaJson, Interactor->GetActionSchema()->ActionSchema, Interactor->GetActionSchemaElement().SchemaElement);
-			check(ActionIndices.Num() == ActionNum);
-
-			for (int Step = 0; Step < Record.StepNum; Step++)
-			{
-				for (int i = 0; i < ActionIndices.Num(); i++)
-				{
-					RecordedActions[StepIdx + Step][i] = ActionsView[Step][ActionIndices[i]];
-				}
-			}
-		}
-		else { continue; }
-
-		EpisodeIdx++;
-		StepIdx += Record.StepNum;
-	}
-
-	check(EpisodeIdx == TotalEpisodeNum);
-	check(StepIdx == TotalStepNum);
-
-	const int32 ObservationSchemaId = 0;
-	const int32 ActionSchemaId = 0;
-
-	// Create Replay Buffer from Records
-	ReplayBuffer = MakeUnique<UE::Learning::FReplayBuffer>();
-	ReplayBuffer->AddRecords(
-		TotalEpisodeNum,
-		TotalStepNum,
-		ObservationSchemaId,
-		ObservationNum,
-		ActionSchemaId,
-		ActionNum,
-		RecordedEpisodeStarts,
-		RecordedEpisodeLengths,
-		RecordedObservations,
-		RecordedActions);
-
 	// We need to setup the trainer prior to sending the config
 	PolicyNetworkId = Trainer->AddNetwork(*Policy->GetPolicyNetworkAsset()->NeuralNetworkData);
 	EncoderNetworkId = Trainer->AddNetwork(*Policy->GetEncoderNetworkAsset()->NeuralNetworkData);
@@ -329,10 +198,11 @@ void ULearningAgentsImitationTrainer::BeginTraining(
 	ReplayBufferId = Trainer->AddReplayBuffer(*ReplayBuffer);
 
 	TSharedRef<FJsonObject> DataConfigObject = CreateDataConfig();
+	TSharedRef<FJsonObject> ModelConfigObject = CreateModelConfig();
 	TSharedRef<FJsonObject> TrainerConfigObject = CreateTrainerConfig(ImitationTrainerTrainingSettings);
 
 	UE_LOG(LogLearning, Display, TEXT("%s: Sending configs..."), *GetName());
-	SendConfigs(DataConfigObject, TrainerConfigObject);
+	SendConfigs(DataConfigObject, ModelConfigObject, TrainerConfigObject);
 
 	UE_LOG(LogLearning, Display, TEXT("%s: Imitation Training Started"), *GetName());
 	UE::Learning::ETrainerResponse Response = UE::Learning::ETrainerResponse::Success;
@@ -378,52 +248,98 @@ void ULearningAgentsImitationTrainer::BeginTraining(
 	bIsTraining = true;
 }
 
-TSharedRef<FJsonObject> ULearningAgentsImitationTrainer::CreateDataConfig() const
+void ULearningAgentsImitationTrainer::ExportToFile(
+	UE::Learning::FFileTrainer& FileTrainer,
+	const ULearningAgentsRecording* Recording,
+	const FLearningAgentsImitationTrainerSettings& ImitationTrainerSettings,
+	const FLearningAgentsImitationTrainerTrainingSettings& ImitationTrainerTrainingSettings,
+	UE::Learning::ETrainerExportFlags ExportFlags)
 {
-	TSharedRef<FJsonObject> ConfigObject = MakeShared<FJsonObject>();
+	using namespace UE::Learning;
 
-	const int32 ObservationSchemaId = 0;
-	const int32 ActionSchemaId = 0;
-
-	// Add Neural Network Config Entries
-	TArray<TSharedPtr<FJsonValue>> NetworkObjects;
+	if (!IsSetup())
+	{
+		UE_LOG(LogLearning, Error, TEXT("%s: Setup not complete."), *GetName());
+		return;
+	}
 
-	// Policy
+	if (IsTraining())
 	{
-		TSharedPtr<FJsonObject> NetworkObject = MakeShared<FJsonObject>();
-		NetworkObject->SetNumberField(TEXT("Id"), PolicyNetworkId);
-		NetworkObject->SetStringField(TEXT("Name"), Policy->GetPolicyNetworkAsset()->GetFName().ToString());
-		NetworkObject->SetNumberField(TEXT("MaxByteNum"), Policy->GetPolicyNetworkAsset()->NeuralNetworkData->GetSnapshotByteNum());
+		UE_LOG(LogLearning, Error, TEXT("%s: Cannot export while training is active."), *GetName());
+		return;
+	}
 
-		TSharedRef<FJsonValueObject> JsonValue = MakeShared<FJsonValueObject>(NetworkObject);
-		NetworkObjects.Add(JsonValue);
+	const bool bNeedReplayBuffers = EnumHasAnyFlags(ExportFlags, ETrainerExportFlags::ReplayBuffers);
+	const bool bNeedNetworks = EnumHasAnyFlags(ExportFlags, ETrainerExportFlags::Networks);
+	const bool bNeedTrainerConfig = EnumHasAnyFlags(ExportFlags, ETrainerExportFlags::TrainerConfig);
+
+	// Networks: register to get IDs (cheap)
+	if (bNeedNetworks)
+	{
+		PolicyNetworkId = FileTrainer.AddNetwork(*Policy->GetPolicyNetworkAsset()->NeuralNetworkData);
+		EncoderNetworkId = FileTrainer.AddNetwork(*Policy->GetEncoderNetworkAsset()->NeuralNetworkData);
+		DecoderNetworkId = FileTrainer.AddNetwork(*Policy->GetDecoderNetworkAsset()->NeuralNetworkData);
 	}
 
-	// Encoder
+	// Replay buffers: build from recording (expensive)
+	if (bNeedReplayBuffers)
 	{
-		TSharedPtr<FJsonObject> NetworkObject = MakeShared<FJsonObject>();
-		NetworkObject->SetNumberField(TEXT("Id"), EncoderNetworkId);
-		NetworkObject->SetStringField(TEXT("Name"), Policy->GetEncoderNetworkAsset()->GetFName().ToString());
-		NetworkObject->SetNumberField(TEXT("MaxByteNum"), Policy->GetEncoderNetworkAsset()->NeuralNetworkData->GetSnapshotByteNum());
-		NetworkObject->SetNumberField(TEXT("InputSchemaId"), ObservationSchemaId);
+		if (!Recording)
+		{
+			UE_LOG(LogLearning, Error, TEXT("%s: Recording is nullptr."), *GetName());
+			return;
+		}
 
-		TSharedRef<FJsonValueObject> JsonValue = MakeShared<FJsonValueObject>(NetworkObject);
-		NetworkObjects.Add(JsonValue);
+		if (Recording->Records.IsEmpty())
+		{
+			UE_LOG(LogLearning, Error, TEXT("%s: Recording is empty!"), *GetName());
+			return;
+		}
+
+		ReplayBuffer = Recording->BuildReplayBuffer(
+			Interactor,
+			ImitationTrainerSettings.RecordingFilterTag,
+			ImitationTrainerTrainingSettings.ValidationSettings.bCheckObservationSchemaCompatibility,
+			ImitationTrainerTrainingSettings.ValidationSettings.bCheckActionSchemaCompatibility);
+
+		if (!ReplayBuffer)
+		{
+			UE_LOG(LogLearning, Warning, TEXT("%s: Recording contains no valid training data."), *GetName());
+			return;
+		}
+
+		ReplayBufferId = FileTrainer.AddReplayBuffer(*ReplayBuffer);
 	}
 
-	// Decoder
+	// Build and send configs
+	TSharedRef<FJsonObject> DataConfigObject = bNeedReplayBuffers ? CreateDataConfig() : MakeShared<FJsonObject>();
+	TSharedRef<FJsonObject> ModelConfigObject = bNeedNetworks ? CreateModelConfig() : MakeShared<FJsonObject>();
+	TSharedRef<FJsonObject> TrainerConfigObject = bNeedTrainerConfig ? CreateTrainerConfig(ImitationTrainerTrainingSettings) : MakeShared<FJsonObject>();
+	FileTrainer.SendConfigs(DataConfigObject, ModelConfigObject, TrainerConfigObject);
+
+	// Send network weights
+	if (bNeedNetworks)
 	{
-		TSharedPtr<FJsonObject> NetworkObject = MakeShared<FJsonObject>();
-		NetworkObject->SetNumberField(TEXT("Id"), DecoderNetworkId);
-		NetworkObject->SetStringField(TEXT("Name"), Policy->GetDecoderNetworkAsset()->GetFName().ToString());
-		NetworkObject->SetNumberField(TEXT("MaxByteNum"), Policy->GetDecoderNetworkAsset()->NeuralNetworkData->GetSnapshotByteNum());
-		NetworkObject->SetNumberField(TEXT("OutputSchemaId"), ActionSchemaId);
+		FileTrainer.SendNetwork(PolicyNetworkId, *Policy->GetPolicyNetworkAsset()->NeuralNetworkData);
+		FileTrainer.SendNetwork(EncoderNetworkId, *Policy->GetEncoderNetworkAsset()->NeuralNetworkData);
+		FileTrainer.SendNetwork(DecoderNetworkId, *Policy->GetDecoderNetworkAsset()->NeuralNetworkData);
+	}
 
-		TSharedRef<FJsonValueObject> JsonValue = MakeShared<FJsonValueObject>(NetworkObject);
-		NetworkObjects.Add(JsonValue);
+	// Send replay buffer data
+	if (bNeedReplayBuffers)
+	{
+		FileTrainer.SendReplayBuffer(ReplayBufferId, *ReplayBuffer);
 	}
 
-	ConfigObject->SetArrayField(TEXT("Networks"), NetworkObjects);
+	UE_LOG(LogLearning, Display, TEXT("%s: Export complete."), *GetName());
+}
+
+TSharedRef<FJsonObject> ULearningAgentsImitationTrainer::CreateDataConfig() const
+{
+	TSharedRef<FJsonObject> ConfigObject = MakeShared<FJsonObject>();
+
+	const int32 ObservationSchemaId = 0;
+	const int32 ActionSchemaId = 0;
 
 	// Add Replay Buffers Config Entries
 	TArray<TSharedPtr<FJsonValue>> ReplayBufferObjects;
@@ -471,6 +387,58 @@ TSharedRef<FJsonObject> ULearningAgentsImitationTrainer::CreateDataConfig() cons
 	return ConfigObject;
 }
 
+TSharedRef<FJsonObject> ULearningAgentsImitationTrainer::CreateModelConfig() const
+{
+	TSharedRef<FJsonObject> ConfigObject = MakeShared<FJsonObject>();
+
+	const int32 ObservationSchemaId = 0;
+	const int32 ActionSchemaId = 0;
+
+	// Add Neural Network Config Entries
+	TArray<TSharedPtr<FJsonValue>> NetworkObjects;
+
+	// Policy
+	{
+		TSharedPtr<FJsonObject> NetworkObject = MakeShared<FJsonObject>();
+		NetworkObject->SetNumberField(TEXT("Id"), PolicyNetworkId);
+		NetworkObject->SetStringField(TEXT("Name"), Policy->GetPolicyNetworkAsset()->GetFName().ToString());
+		NetworkObject->SetNumberField(TEXT("MaxByteNum"), Policy->GetPolicyNetworkAsset()->NeuralNetworkData->GetSnapshotByteNum());
+
+		TSharedRef<FJsonValueObject> JsonValue = MakeShared<FJsonValueObject>(NetworkObject);
+		NetworkObjects.Add(JsonValue);
+	}
+
+	// Encoder
+	{
+		TSharedPtr<FJsonObject> NetworkObject = MakeShared<FJsonObject>();
+		NetworkObject->SetNumberField(TEXT("Id"), EncoderNetworkId);
+		NetworkObject->SetStringField(TEXT("Name"), Policy->GetEncoderNetworkAsset()->GetFName().ToString());
+		NetworkObject->SetNumberField(TEXT("MaxByteNum"), Policy->GetEncoderNetworkAsset()->NeuralNetworkData->GetSnapshotByteNum());
+		NetworkObject->SetNumberField(TEXT("InputSchemaId"), ObservationSchemaId);
+
+		TSharedRef<FJsonValueObject> JsonValue = MakeShared<FJsonValueObject>(NetworkObject);
+		NetworkObjects.Add(JsonValue);
+	}
+
+	// Decoder
+	{
+		TSharedPtr<FJsonObject> NetworkObject = MakeShared<FJsonObject>();
+		NetworkObject->SetNumberField(TEXT("Id"), DecoderNetworkId);
+		NetworkObject->SetStringField(TEXT("Name"), Policy->GetDecoderNetworkAsset()->GetFName().ToString());
+		NetworkObject->SetNumberField(TEXT("MaxByteNum"), Policy->GetDecoderNetworkAsset()->NeuralNetworkData->GetSnapshotByteNum());
+		NetworkObject->SetNumberField(TEXT("OutputSchemaId"), ActionSchemaId);
+
+		TSharedRef<FJsonValueObject> JsonValue = MakeShared<FJsonValueObject>(NetworkObject);
+		NetworkObjects.Add(JsonValue);
+	}
+
+	ConfigObject->SetArrayField(TEXT("Networks"), NetworkObjects);
+
+	ConfigObject->SetObjectField(TEXT("PolicySettings"), Policy->AsJsonConfig());
+
+	return ConfigObject;
+}
+
 TSharedRef<FJsonObject> ULearningAgentsImitationTrainer::CreateTrainerConfig(const FLearningAgentsImitationTrainerTrainingSettings& TrainingSettings) const
 {
 	TSharedRef<FJsonObject> ConfigObject = MakeShared<FJsonObject>();
@@ -481,8 +449,7 @@ TSharedRef<FJsonObject> ULearningAgentsImitationTrainer::CreateTrainerConfig(con
 
 	// Add Imitation Specific Config Entries
 	ConfigObject->SetObjectField(TEXT("BehaviorCloningSettings"), TrainingSettings.AsJsonConfig());
-	ConfigObject->SetObjectField(TEXT("PolicySettings"), Policy->AsJsonConfig());
-	
+
 	// Add Normalization Choice
 	ConfigObject->SetBoolField(TEXT("EmbedNormalizationsInWeights"), TrainingSettings.bEmbedNormalizationsInWeights);
 
@@ -495,10 +462,10 @@ TSharedRef<FJsonObject> ULearningAgentsImitationTrainer::CreateTrainerConfig(con
 	return ConfigObject;
 }
 
-void ULearningAgentsImitationTrainer::SendConfigs(const TSharedRef<FJsonObject>& DataConfigObject, const TSharedRef<FJsonObject>& TrainerConfigObject)
+void ULearningAgentsImitationTrainer::SendConfigs(const TSharedRef<FJsonObject>& DataConfigObject, const TSharedRef<FJsonObject>& ModelConfigObject, const TSharedRef<FJsonObject>& TrainerConfigObject)
 {
 	UE::Learning::ETrainerResponse Response = UE::Learning::ETrainerResponse::Success;
-	Response = Trainer->SendConfigs(DataConfigObject, TrainerConfigObject);
+	Response = Trainer->SendConfigs(DataConfigObject, ModelConfigObject, TrainerConfigObject);
 
 	if (Response != UE::Learning::ETrainerResponse::Success)
 	{
diff --git a/Engine/Plugins/Experimental/LearningAgents/Source/LearningAgentsTraining/Private/LearningAgentsPPOTrainer.cpp b/Engine/Plugins/Experimental/LearningAgents/Source/LearningAgentsTraining/Private/LearningAgentsPPOTrainer.cpp
--- a/Engine/Plugins/Experimental/LearningAgents/Source/LearningAgentsTraining/Private/LearningAgentsPPOTrainer.cpp
+++ b/Engine/Plugins/Experimental/LearningAgents/Source/LearningAgentsTraining/Private/LearningAgentsPPOTrainer.cpp
@@ -467,8 +467,9 @@ TSharedRef<FJsonObject> ULearningAgentsPPOTrainer::CreateTrainerConfig(const FLe
 
 void ULearningAgentsPPOTrainer::SendConfigs(const TSharedRef<FJsonObject>& DataConfigObject, const TSharedRef<FJsonObject>& TrainerConfigObject)
 {
+	TSharedRef<FJsonObject> ModelConfigObject = MakeShared<FJsonObject>();
 	UE::Learning::ETrainerResponse Response = UE::Learning::ETrainerResponse::Success;
-	Response = Trainer->SendConfigs(DataConfigObject, TrainerConfigObject);
+	Response = Trainer->SendConfigs(DataConfigObject, ModelConfigObject, TrainerConfigObject);
 	
 	if (Response != UE::Learning::ETrainerResponse::Success)
 	{
diff --git a/Engine/Plugins/Experimental/LearningAgents/Source/LearningAgentsTraining/Private/LearningAgentsRecording.cpp b/Engine/Plugins/Experimental/LearningAgents/Source/LearningAgentsTraining/Private/LearningAgentsRecording.cpp
--- a/Engine/Plugins/Experimental/LearningAgents/Source/LearningAgentsTraining/Private/LearningAgentsRecording.cpp
+++ b/Engine/Plugins/Experimental/LearningAgents/Source/LearningAgentsTraining/Private/LearningAgentsRecording.cpp
@@ -2,7 +2,11 @@
 
 #include "LearningAgentsRecording.h"
 
+#include "LearningAgentsInteractor.h"
+
+#include "LearningExperience.h"
 #include "LearningLog.h"
+#include "LearningTrainer.h"
 
 #include "UObject/Package.h"
 #include "Misc/FileHelper.h"
@@ -379,6 +383,159 @@ void ULearningAgentsRecording::GetActionVector(TArray<float>& OutActionVector, i
 	OutActionCompatibilityHash = Records[Record].ActionCompatibilityHash;
 }
 
+TUniquePtr<UE::Learning::FReplayBuffer> ULearningAgentsRecording::BuildReplayBuffer(
+	const ULearningAgentsInteractor* Interactor,
+	const FGameplayTag& RecordingFilterTag,
+	bool bCheckObservationSchemaCompatibility,
+	bool bCheckActionSchemaCompatibility) const
+{
+	using namespace UE::Learning;
+
+	// Count valid episodes and steps
+
+	int32 TotalEpisodeNum = 0;
+	int32 TotalStepNum = 0;
+	for (const FLearningAgentsRecord& Record : Records)
+	{
+		if (RecordingFilterTag.IsValid() && !Record.Tags.HasTag(RecordingFilterTag))
+		{
+			UE_LOG(LogLearning, Log, TEXT("Skipping record %s because its missing tag %s."), *Record.Name, *RecordingFilterTag.ToString());
+			continue;
+		}
+
+		if (!Schemas.Contains(Record.SchemaName))
+		{
+			UE_LOG(LogLearning, Error, TEXT("Schema name %s not found in schema list. Is the schema added and are the names correct?"), *Record.SchemaName.ToString());
+			continue;
+		}
+
+		if (bCheckObservationSchemaCompatibility && !Trainer::IsObservationSchemaSubsetCompatible(Schemas[Record.SchemaName].ObservationSchemaJson, Interactor->GetObservationSchema()->ObservationSchema, Interactor->GetObservationSchemaElement().SchemaElement))
+		{
+			UE_LOG(LogLearning, Warning, TEXT("Record has incompatible observation schema."));
+			continue;
+		}
+
+		if (bCheckActionSchemaCompatibility && !Trainer::IsActionSchemaSubsetCompatible(Schemas[Record.SchemaName].ActionSchemaJson, Interactor->GetActionSchema()->ActionSchema, Interactor->GetActionSchemaElement().SchemaElement))
+		{
+			UE_LOG(LogLearning, Warning, TEXT("Record has incompatible action schema."));
+			continue;
+		}
+
+		TotalEpisodeNum++;
+		TotalStepNum += Record.StepNum;
+	}
+
+	if (TotalStepNum == 0)
+	{
+		UE_LOG(LogLearning, Warning, TEXT("Recording contains no valid training data."));
+		return nullptr;
+	}
+
+	// Sizes
+
+	const int32 ObservationNum = Interactor->GetObservationVectorSize();
+	const int32 ActionNum = Interactor->GetActionVectorSize();
+
+	// Copy into flat arrays
+
+	TLearningArray<1, int32> RecordedEpisodeStarts;
+	TLearningArray<1, int32> RecordedEpisodeLengths;
+	TLearningArray<2, float> RecordedObservations;
+	TLearningArray<2, float> RecordedActions;
+
+	RecordedEpisodeStarts.SetNumUninitialized({ TotalEpisodeNum });
+	RecordedEpisodeLengths.SetNumUninitialized({ TotalEpisodeNum });
+	RecordedObservations.SetNumUninitialized({ TotalStepNum, ObservationNum });
+	RecordedActions.SetNumUninitialized({ TotalStepNum, ActionNum });
+
+	int32 EpisodeIdx = 0;
+	int32 StepIdx = 0;
+	for (const FLearningAgentsRecord& Record : Records)
+	{
+		if (RecordingFilterTag.IsValid() && !Record.Tags.HasTag(RecordingFilterTag))
+		{
+			continue;
+		}
+
+		if (!Schemas.Contains(Record.SchemaName))
+		{
+			continue;
+		}
+
+		const bool bObservationSchemaSubsetCompatible = (!bCheckObservationSchemaCompatibility) || Trainer::IsObservationSchemaSubsetCompatible(Schemas[Record.SchemaName].ObservationSchemaJson, Interactor->GetObservationSchema()->ObservationSchema, Interactor->GetObservationSchemaElement().SchemaElement);
+		const bool bActionSchemaSubsetCompatible = (!bCheckActionSchemaCompatibility) || Trainer::IsActionSchemaSubsetCompatible(Schemas[Record.SchemaName].ActionSchemaJson, Interactor->GetActionSchema()->ActionSchema, Interactor->GetActionSchemaElement().SchemaElement);
+
+		if (Record.ObservationDimNum == ObservationNum && Record.ActionDimNum == ActionNum)
+		{
+			TLearningArrayView<2, const float> ObservationsView(Record.Blob->ObservationData.GetData(), { Record.StepNum, Record.ObservationDimNum });
+			TLearningArrayView<2, const float> ActionsView(Record.Blob->ActionData.GetData(), { Record.StepNum, Record.ActionDimNum });
+
+			RecordedEpisodeStarts[EpisodeIdx] = StepIdx;
+			RecordedEpisodeLengths[EpisodeIdx] = Record.StepNum;
+			Array::Copy(RecordedObservations.Slice(StepIdx, Record.StepNum), ObservationsView);
+			Array::Copy(RecordedActions.Slice(StepIdx, Record.StepNum), ActionsView);
+		}
+		else if (bObservationSchemaSubsetCompatible && bActionSchemaSubsetCompatible)
+		{
+			TLearningArrayView<2, const float> ObservationsView(Record.Blob->ObservationData.GetData(), { Record.StepNum, Record.ObservationDimNum });
+			TLearningArrayView<2, const float> ActionsView(Record.Blob->ActionData.GetData(), { Record.StepNum, Record.ActionDimNum });
+
+			RecordedEpisodeStarts[EpisodeIdx] = StepIdx;
+			RecordedEpisodeLengths[EpisodeIdx] = Record.StepNum;
+
+			TArray<int32> ObservationIndices = Trainer::ComputeObservationSchemaSubsetIndices(
+				Schemas[Record.SchemaName].ObservationSchemaJson, Interactor->GetObservationSchema()->ObservationSchema, Interactor->GetObservationSchemaElement().SchemaElement);
+			check(ObservationIndices.Num() == ObservationNum);
+
+			for (int32 Step = 0; Step < Record.StepNum; Step++)
+			{
+				for (int32 i = 0; i < ObservationIndices.Num(); i++)
+				{
+					RecordedObservations[StepIdx + Step][i] = ObservationsView[Step][ObservationIndices[i]];
+				}
+			}
+
+			TArray<int32> ActionIndices = Trainer::ComputeActionSchemaSubsetIndices(
+				Schemas[Record.SchemaName].ActionSchemaJson, Interactor->GetActionSchema()->ActionSchema, Interactor->GetActionSchemaElement().SchemaElement);
+			check(ActionIndices.Num() == ActionNum);
+
+			for (int32 Step = 0; Step < Record.StepNum; Step++)
+			{
+				for (int32 i = 0; i < ActionIndices.Num(); i++)
+				{
+					RecordedActions[StepIdx + Step][i] = ActionsView[Step][ActionIndices[i]];
+				}
+			}
+		}
+		else { continue; }
+
+		EpisodeIdx++;
+		StepIdx += Record.StepNum;
+	}
+
+	check(EpisodeIdx == TotalEpisodeNum);
+	check(StepIdx == TotalStepNum);
+
+	const int32 ObservationSchemaId = 0;
+	const int32 ActionSchemaId = 0;
+
+	// Create Replay Buffer from Records
+	TUniquePtr<FReplayBuffer> ReplayBuffer = MakeUnique<FReplayBuffer>();
+	ReplayBuffer->AddRecords(
+		TotalEpisodeNum,
+		TotalStepNum,
+		ObservationSchemaId,
+		ObservationNum,
+		ActionSchemaId,
+		ActionNum,
+		RecordedEpisodeStarts,
+		RecordedEpisodeLengths,
+		RecordedObservations,
+		RecordedActions);
+
+	return ReplayBuffer;
+}
+
 void ULearningAgentsRecording::ForceMarkDirty()
 {
 	// Manually mark the package as dirty since just using `Modify` prevents 
diff --git a/Engine/Plugins/Experimental/LearningAgents/Source/LearningAgentsTraining/Public/LearningAgentsFlowMatchingTrainer.h b/Engine/Plugins/Experimental/LearningAgents/Source/LearningAgentsTraining/Public/LearningAgentsFlowMatchingTrainer.h
--- a/Engine/Plugins/Experimental/LearningAgents/Source/LearningAgentsTraining/Public/LearningAgentsFlowMatchingTrainer.h
+++ b/Engine/Plugins/Experimental/LearningAgents/Source/LearningAgentsTraining/Public/LearningAgentsFlowMatchingTrainer.h
@@ -15,6 +15,7 @@
 
 namespace UE::Learning
 {
+	struct FFileTrainer;
 	struct FReplayBuffer;
 	struct IExternalTrainer;
 }
@@ -276,6 +277,22 @@ public:
 		const FLearningAgentsFlowMatchingTrainerTrainingSettings& FlowMatchingTrainerTrainingSettings = FLearningAgentsFlowMatchingTrainerTrainingSettings(),
 		const FLearningAgentsTrainerProcessSettings& FlowMatchingTrainerPathSettings = FLearningAgentsTrainerProcessSettings());
 
+	/**
+	 * Exports training materials to file, only performing work needed for the requested export flags.
+	 * Skips expensive replay buffer construction when not exporting replay buffers.
+	 *
+	 * @param Recording							The data to export.
+	 * @param FlowMatchingTrainerSettings		The settings for this trainer.
+	 * @param FlowMatchingTrainerTrainingSettings	The training settings for this network.
+	 * @param ExportFlags						Which sections to export.
+	 */
+	UE_API void ExportToFile(
+		UE::Learning::FFileTrainer& FileTrainer,
+		const ULearningAgentsRecording* Recording,
+		const FLearningAgentsFlowMatchingTrainerSettings& FlowMatchingTrainerSettings,
+		const FLearningAgentsFlowMatchingTrainerTrainingSettings& FlowMatchingTrainerTrainingSettings,
+		UE::Learning::ETrainerExportFlags ExportFlags);
+
 // ----- Private Data -----
 private:
 
@@ -300,8 +317,9 @@ private:
 
 	UE_API TSharedRef<FJsonObject> CreateSchemaConfig(const int32 ObservationSchemaId = 0, const int32 ActionSchemaId = 0) const;
 	UE_API TSharedRef<FJsonObject> CreateDataConfig() const;
+	UE_API TSharedRef<FJsonObject> CreateModelConfig() const;
 	UE_API TSharedRef<FJsonObject> CreateTrainerConfig(const FLearningAgentsFlowMatchingTrainerTrainingSettings& TrainingSettings) const;
-	UE_API void SendConfigs(const TSharedRef<FJsonObject>& DataConfigObject, const TSharedRef<FJsonObject>& TrainerConfigObject);
+	UE_API void SendConfigs(const TSharedRef<FJsonObject>& DataConfigObject, const TSharedRef<FJsonObject>& ModelConfigObject, const TSharedRef<FJsonObject>& TrainerConfigObject);
 
 	UE_API void DoneTraining();
 
diff --git a/Engine/Plugins/Experimental/LearningAgents/Source/LearningAgentsTraining/Public/LearningAgentsImitationTrainer.h b/Engine/Plugins/Experimental/LearningAgents/Source/LearningAgentsTraining/Public/LearningAgentsImitationTrainer.h
--- a/Engine/Plugins/Experimental/LearningAgents/Source/LearningAgentsTraining/Public/LearningAgentsImitationTrainer.h
+++ b/Engine/Plugins/Experimental/LearningAgents/Source/LearningAgentsTraining/Public/LearningAgentsImitationTrainer.h
@@ -15,6 +15,7 @@
 
 namespace UE::Learning
 {
+	struct FFileTrainer;
 	struct FReplayBuffer;
 	struct IExternalTrainer;
 }
@@ -321,6 +322,22 @@ public:
 		const FLearningAgentsImitationTrainerTrainingSettings& ImitationTrainerTrainingSettings = FLearningAgentsImitationTrainerTrainingSettings(),
 		const FLearningAgentsTrainerProcessSettings& ImitationTrainerPathSettings = FLearningAgentsTrainerProcessSettings());
 
+	/**
+	 * Exports training materials to file, only performing work needed for the requested export flags.
+	 * Skips expensive replay buffer construction when not exporting replay buffers.
+	 *
+	 * @param Recording							The data to export.
+	 * @param ImitationTrainerSettings			The settings for this trainer.
+	 * @param ImitationTrainerTrainingSettings	The training settings for this network.
+	 * @param ExportFlags						Which sections to export.
+	 */
+	UE_API void ExportToFile(
+		UE::Learning::FFileTrainer& FileTrainer,
+		const ULearningAgentsRecording* Recording,
+		const FLearningAgentsImitationTrainerSettings& ImitationTrainerSettings,
+		const FLearningAgentsImitationTrainerTrainingSettings& ImitationTrainerTrainingSettings,
+		UE::Learning::ETrainerExportFlags ExportFlags);
+
 // ----- Private Data -----
 private:
 
@@ -344,8 +361,9 @@ private:
 	bool bHasTrainingFailed = false;
 
 	UE_API TSharedRef<FJsonObject> CreateDataConfig() const;
+	UE_API TSharedRef<FJsonObject> CreateModelConfig() const;
 	UE_API TSharedRef<FJsonObject> CreateTrainerConfig(const FLearningAgentsImitationTrainerTrainingSettings& TrainingSettings) const;
-	UE_API void SendConfigs(const TSharedRef<FJsonObject>& DataConfigObject, const TSharedRef<FJsonObject>& TrainerConfigObject);
+	UE_API void SendConfigs(const TSharedRef<FJsonObject>& DataConfigObject, const TSharedRef<FJsonObject>& ModelConfigObject, const TSharedRef<FJsonObject>& TrainerConfigObject);
 
 	UE_API void DoneTraining();
 
diff --git a/Engine/Plugins/Experimental/LearningAgents/Source/LearningAgentsTraining/Public/LearningAgentsRecording.h b/Engine/Plugins/Experimental/LearningAgents/Source/LearningAgentsTraining/Public/LearningAgentsRecording.h
--- a/Engine/Plugins/Experimental/LearningAgents/Source/LearningAgentsTraining/Public/LearningAgentsRecording.h
+++ b/Engine/Plugins/Experimental/LearningAgents/Source/LearningAgentsTraining/Public/LearningAgentsRecording.h
@@ -9,10 +9,19 @@
 #include "GameplayTagContainer.h"
 #include "UObject/NameTypes.h"
 
+#include "Templates/UniquePtr.h"
+
 #include "LearningAgentsRecording.generated.h"
 
 #define UE_API LEARNINGAGENTSTRAINING_API
 
+namespace UE::Learning
+{
+	struct FReplayBuffer;
+}
+
+class ULearningAgentsInteractor;
+
 namespace UE::Learning::Agents::Recording
 {
 	UE_API int32 GetMagicNumber();
@@ -186,6 +195,22 @@ public:
 
 public:
 
+	/**
+	 * Builds a replay buffer from this recording, filtering by tag and validating schema compatibility.
+	 * Returns nullptr if no valid data is found.
+	 *
+	 * @param Interactor							The interactor whose observation/action schemas define the target layout.
+	 * @param RecordingFilterTag					Tag used to filter records. If invalid, no filtering is applied.
+	 * @param bCheckObservationSchemaCompatibility	Whether to validate observation schema subset compatibility.
+	 * @param bCheckActionSchemaCompatibility		Whether to validate action schema subset compatibility.
+	 * @return A replay buffer containing the copied data, or nullptr if no valid records were found.
+	 */
+	UE_API TUniquePtr<UE::Learning::FReplayBuffer> BuildReplayBuffer(
+		const ULearningAgentsInteractor* Interactor,
+		const FGameplayTag& RecordingFilterTag,
+		bool bCheckObservationSchemaCompatibility,
+		bool bCheckActionSchemaCompatibility) const;
+
 	/** Marks this asset as modified even during PIE */
 	UE_API void ForceMarkDirty();
 
diff --git a/Engine/Plugins/Experimental/LearningAgents/Source/LearningAgentsTraining/Public/LearningAgentsTrainer.h b/Engine/Plugins/Experimental/LearningAgents/Source/LearningAgentsTraining/Public/LearningAgentsTrainer.h
--- a/Engine/Plugins/Experimental/LearningAgents/Source/LearningAgentsTraining/Public/LearningAgentsTrainer.h
+++ b/Engine/Plugins/Experimental/LearningAgents/Source/LearningAgentsTraining/Public/LearningAgentsTrainer.h
@@ -156,7 +156,7 @@ public:
 
 	/** Trainer file name. The name of the python file to use for training. Do NOT include the '.py' file extension. */
 	UPROPERTY(EditAnywhere, Category = "LearningAgents")
-	FString TrainerFileName = TEXT("learning_agents.train_ppo");
+	FString TrainerFileName = TEXT("learning_core.train_behavior_cloning");
 
 public:
 
diff --git a/Engine/Plugins/Experimental/LearningAgents/Source/LearningAgentsTrainingEditor/Private/LearningAgentsFlowMatchingTrainerEditor.cpp b/Engine/Plugins/Experimental/LearningAgents/Source/LearningAgentsTrainingEditor/Private/LearningAgentsFlowMatchingTrainerEditor.cpp
--- a/Engine/Plugins/Experimental/LearningAgents/Source/LearningAgentsTrainingEditor/Private/LearningAgentsFlowMatchingTrainerEditor.cpp
+++ b/Engine/Plugins/Experimental/LearningAgents/Source/LearningAgentsTrainingEditor/Private/LearningAgentsFlowMatchingTrainerEditor.cpp
@@ -50,6 +50,32 @@ void ALearningAgentsFlowMatchingTrainerEditor::StartTraining()
 	);
 }
 
+void ALearningAgentsFlowMatchingTrainerEditor::StartExport(UE::Learning::ETrainerExportFlags Flags)
+{
+	ExportFlags = Flags;
+
+	ULearningAgentsFlowMatchingTrainer* Trainer = Cast<ULearningAgentsFlowMatchingTrainer>(SetupTraining());
+	if (!Trainer)
+	{
+		UE_LOGFMT(LogLearning, Error, "{Name}: Export failed. Invalid flow matching trainer. Was the flow matching trainer set with SetupTraining?", ("Name", GetName()));
+		ExportFlags = UE::Learning::ETrainerExportFlags::All;
+		return;
+	}
+
+	if (!FileTrainer)
+	{
+		UE_LOGFMT(LogLearning, Error, "{Name}: Export failed. No file communicator. Use MakeFileCommunicator in your SetupTraining Blueprint.", ("Name", GetName()));
+		ExportFlags = UE::Learning::ETrainerExportFlags::All;
+		return;
+	}
+
+	UE_LOGFMT(LogLearning, Display, "{Name}: Exporting training materials...", ("Name", GetName()));
+
+	Trainer->ExportToFile(*FileTrainer, Recording, FlowMatchingTrainerSettings, FlowMatchingTrainerTrainingSettings, Flags);
+
+	ExportFlags = UE::Learning::ETrainerExportFlags::All;
+}
+
 void ALearningAgentsFlowMatchingTrainerEditor::StopTraining()
 {
 	if (!LearningAgentsFlowMatchingTrainer)
diff --git a/Engine/Plugins/Experimental/LearningAgents/Source/LearningAgentsTrainingEditor/Private/LearningAgentsImitationTrainerEditor.cpp b/Engine/Plugins/Experimental/LearningAgents/Source/LearningAgentsTrainingEditor/Private/LearningAgentsImitationTrainerEditor.cpp
--- a/Engine/Plugins/Experimental/LearningAgents/Source/LearningAgentsTrainingEditor/Private/LearningAgentsImitationTrainerEditor.cpp
+++ b/Engine/Plugins/Experimental/LearningAgents/Source/LearningAgentsTrainingEditor/Private/LearningAgentsImitationTrainerEditor.cpp
@@ -49,6 +49,32 @@ void ALearningAgentsImitationTrainerEditor::StartTraining()
 	);
 }
 
+void ALearningAgentsImitationTrainerEditor::StartExport(UE::Learning::ETrainerExportFlags Flags)
+{
+	ExportFlags = Flags;
+
+	ULearningAgentsImitationTrainer* Trainer = Cast<ULearningAgentsImitationTrainer>(SetupTraining());
+	if (!Trainer)
+	{
+		UE_LOG(LogLearning, Error, TEXT("%s: Export failed. Invalid imitation trainer. Was the imitation trainer set with SetupTraining?"), *GetName());
+		ExportFlags = UE::Learning::ETrainerExportFlags::All;
+		return;
+	}
+
+	if (!FileTrainer)
+	{
+		UE_LOG(LogLearning, Error, TEXT("%s: Export failed. No file communicator. Use MakeFileCommunicator in your SetupTraining Blueprint."), *GetName());
+		ExportFlags = UE::Learning::ETrainerExportFlags::All;
+		return;
+	}
+
+	UE_LOG(LogLearning, Display, TEXT("%s: Exporting training materials..."), *GetName());
+
+	Trainer->ExportToFile(*FileTrainer, Recording, ImitationTrainerSettings, ImitationTrainerTrainingSettings, Flags);
+
+	ExportFlags = UE::Learning::ETrainerExportFlags::All;
+}
+
 void ALearningAgentsImitationTrainerEditor::StopTraining()
 {
 	if (!LearningAgentsImitationTrainer)
diff --git a/Engine/Plugins/Experimental/LearningAgents/Source/LearningAgentsTrainingEditor/Private/LearningAgentsTrainerEditorBase.cpp b/Engine/Plugins/Experimental/LearningAgents/Source/LearningAgentsTrainingEditor/Private/LearningAgentsTrainerEditorBase.cpp
--- a/Engine/Plugins/Experimental/LearningAgents/Source/LearningAgentsTrainingEditor/Private/LearningAgentsTrainerEditorBase.cpp
+++ b/Engine/Plugins/Experimental/LearningAgents/Source/LearningAgentsTrainingEditor/Private/LearningAgentsTrainerEditorBase.cpp
@@ -14,7 +14,8 @@ ALearningAgentsTrainerEditorBase::ALearningAgentsTrainerEditorBase()
 
 FLearningAgentsCommunicator ALearningAgentsTrainerEditorBase::MakeFileCommunicator(FDirectoryPath EditorIntermediateRelativePath)
 {
+	FileTrainer = MakeShared<UE::Learning::FFileTrainer>(UE::Learning::Trainer::GetIntermediatePath(EditorIntermediateRelativePath.Path), ExportFlags);
 	FLearningAgentsCommunicator Communicator;
-	Communicator.Trainer = MakeShared<UE::Learning::FFileTrainer>(UE::Learning::Trainer::GetIntermediatePath(EditorIntermediateRelativePath.Path));
+	Communicator.Trainer = FileTrainer;
 	return Communicator;
 }
diff --git a/Engine/Plugins/Experimental/LearningAgents/Source/LearningAgentsTrainingEditor/Private/LearningAgentsTrainerEditorDetails.cpp b/Engine/Plugins/Experimental/LearningAgents/Source/LearningAgentsTrainingEditor/Private/LearningAgentsTrainerEditorDetails.cpp
--- a/Engine/Plugins/Experimental/LearningAgents/Source/LearningAgentsTrainingEditor/Private/LearningAgentsTrainerEditorDetails.cpp
+++ b/Engine/Plugins/Experimental/LearningAgents/Source/LearningAgentsTrainingEditor/Private/LearningAgentsTrainerEditorDetails.cpp
@@ -2,6 +2,7 @@
 
 #include "LearningAgentsTrainerEditorDetails.h"
 #include "LearningAgentsTrainerEditorBase.h"
+#include "LearningExternalTrainer.h"
 #include "LearningLog.h"
 #include "DetailCategoryBuilder.h"
 #include "DetailWidgetRow.h"
@@ -72,6 +73,66 @@ void FLearningAgentsTrainerEditorDetails::CustomizeDetails(IDetailLayoutBuilder&
 						})
 				]
 		];
+
+	auto IsNotTrainingLambda = [this]()
+	{
+		if (EditedObjects.Num() > 0)
+		{
+			if (ALearningAgentsTrainerEditorBase* TrainerEditor = Cast<ALearningAgentsTrainerEditorBase>(EditedObjects[0].Get()))
+			{
+				return !TrainerEditor->IsTraining();
+			}
+		}
+		return false;
+	};
+
+	Category.AddCustomRow(LOCTEXT("ExportCategory", "Export"))
+		.NameContent()
+		[
+			SNew(STextBlock)
+				.Text(LOCTEXT("ExportLabel", "Export"))
+				.Font(DetailBuilder.GetDetailFont())
+		]
+		.ValueContent()
+		[
+			SNew(SHorizontalBox)
+				+ SHorizontalBox::Slot()
+				.AutoWidth()
+				.Padding(2.0f, 0.0f)
+				[
+					SNew(SButton)
+						.Text(LOCTEXT("ExportAllButton", "All"))
+						.OnClicked(FOnClicked::CreateSP(this, &FLearningAgentsTrainerEditorDetails::OnExportAllClicked))
+						.IsEnabled_Lambda(IsNotTrainingLambda)
+				]
+				+ SHorizontalBox::Slot()
+				.AutoWidth()
+				.Padding(2.0f, 0.0f)
+				[
+					SNew(SButton)
+						.Text(LOCTEXT("ExportNetworksButton", "Networks"))
+						.OnClicked(FOnClicked::CreateSP(this, &FLearningAgentsTrainerEditorDetails::OnExportNetworksClicked))
+						.IsEnabled_Lambda(IsNotTrainingLambda)
+				]
+				+ SHorizontalBox::Slot()
+				.AutoWidth()
+				.Padding(2.0f, 0.0f)
+				[
+					SNew(SButton)
+						.Text(LOCTEXT("ExportReplayBuffersButton", "Replay Buffers"))
+						.OnClicked(FOnClicked::CreateSP(this, &FLearningAgentsTrainerEditorDetails::OnExportReplayBuffersClicked))
+						.IsEnabled_Lambda(IsNotTrainingLambda)
+				]
+				+ SHorizontalBox::Slot()
+				.AutoWidth()
+				.Padding(2.0f, 0.0f)
+				[
+					SNew(SButton)
+						.Text(LOCTEXT("ExportTrainerConfigButton", "Trainer Config"))
+						.OnClicked(FOnClicked::CreateSP(this, &FLearningAgentsTrainerEditorDetails::OnExportTrainerConfigClicked))
+						.IsEnabled_Lambda(IsNotTrainingLambda)
+				]
+		];
 }
 
 FReply FLearningAgentsTrainerEditorDetails::OnRunClicked()
@@ -98,4 +159,52 @@ FReply FLearningAgentsTrainerEditorDetails::OnStopClicked()
 	return FReply::Handled();
 }
 
+FReply FLearningAgentsTrainerEditorDetails::OnExportAllClicked()
+{
+	if (EditedObjects.Num() > 0)
+	{
+		if (ALearningAgentsTrainerEditorBase* TrainerEditor = Cast<ALearningAgentsTrainerEditorBase>(EditedObjects[0].Get()))
+		{
+			TrainerEditor->StartExport(UE::Learning::ETrainerExportFlags::All);
+		}
+	}
+	return FReply::Handled();
+}
+
+FReply FLearningAgentsTrainerEditorDetails::OnExportNetworksClicked()
+{
+	if (EditedObjects.Num() > 0)
+	{
+		if (ALearningAgentsTrainerEditorBase* TrainerEditor = Cast<ALearningAgentsTrainerEditorBase>(EditedObjects[0].Get()))
+		{
+			TrainerEditor->StartExport(UE::Learning::ETrainerExportFlags::Networks);
+		}
+	}
+	return FReply::Handled();
+}
+
+FReply FLearningAgentsTrainerEditorDetails::OnExportReplayBuffersClicked()
+{
+	if (EditedObjects.Num() > 0)
+	{
+		if (ALearningAgentsTrainerEditorBase* TrainerEditor = Cast<ALearningAgentsTrainerEditorBase>(EditedObjects[0].Get()))
+		{
+			TrainerEditor->StartExport(UE::Learning::ETrainerExportFlags::ReplayBuffers);
+		}
+	}
+	return FReply::Handled();
+}
+
+FReply FLearningAgentsTrainerEditorDetails::OnExportTrainerConfigClicked()
+{
+	if (EditedObjects.Num() > 0)
+	{
+		if (ALearningAgentsTrainerEditorBase* TrainerEditor = Cast<ALearningAgentsTrainerEditorBase>(EditedObjects[0].Get()))
+		{
+			TrainerEditor->StartExport(UE::Learning::ETrainerExportFlags::TrainerConfig);
+		}
+	}
+	return FReply::Handled();
+}
+
 #undef LOCTEXT_NAMESPACE
diff --git a/Engine/Plugins/Experimental/LearningAgents/Source/LearningAgentsTrainingEditor/Public/LearningAgentsFlowMatchingTrainerEditor.h b/Engine/Plugins/Experimental/LearningAgents/Source/LearningAgentsTrainingEditor/Public/LearningAgentsFlowMatchingTrainerEditor.h
--- a/Engine/Plugins/Experimental/LearningAgents/Source/LearningAgentsTrainingEditor/Public/LearningAgentsFlowMatchingTrainerEditor.h
+++ b/Engine/Plugins/Experimental/LearningAgents/Source/LearningAgentsTrainingEditor/Public/LearningAgentsFlowMatchingTrainerEditor.h
@@ -28,6 +28,8 @@ public:
 
 	UE_API virtual bool IsTraining() const override;
 
+	UE_API virtual void StartExport(UE::Learning::ETrainerExportFlags Flags) override;
+
 	UPROPERTY(BlueprintReadWrite, Category = "LearningAgents")
 	TObjectPtr<ULearningAgentsFlowMatchingTrainer> LearningAgentsFlowMatchingTrainer;
 
diff --git a/Engine/Plugins/Experimental/LearningAgents/Source/LearningAgentsTrainingEditor/Public/LearningAgentsImitationTrainerEditor.h b/Engine/Plugins/Experimental/LearningAgents/Source/LearningAgentsTrainingEditor/Public/LearningAgentsImitationTrainerEditor.h
--- a/Engine/Plugins/Experimental/LearningAgents/Source/LearningAgentsTrainingEditor/Public/LearningAgentsImitationTrainerEditor.h
+++ b/Engine/Plugins/Experimental/LearningAgents/Source/LearningAgentsTrainingEditor/Public/LearningAgentsImitationTrainerEditor.h
@@ -29,6 +29,8 @@ public:
 
 	UE_API virtual bool IsTraining() const override;
 
+	UE_API virtual void StartExport(UE::Learning::ETrainerExportFlags Flags) override;
+
 	UPROPERTY(BlueprintReadWrite, Category = "LearningAgents")
 	TObjectPtr<ULearningAgentsImitationTrainer> LearningAgentsImitationTrainer;
 
diff --git a/Engine/Plugins/Experimental/LearningAgents/Source/LearningAgentsTrainingEditor/Public/LearningAgentsTrainerEditorBase.h b/Engine/Plugins/Experimental/LearningAgents/Source/LearningAgentsTrainingEditor/Public/LearningAgentsTrainerEditorBase.h
--- a/Engine/Plugins/Experimental/LearningAgents/Source/LearningAgentsTrainingEditor/Public/LearningAgentsTrainerEditorBase.h
+++ b/Engine/Plugins/Experimental/LearningAgents/Source/LearningAgentsTrainingEditor/Public/LearningAgentsTrainerEditorBase.h
@@ -8,6 +8,7 @@
 #include "LearningAgentsCommunicator.h"
 #include "LearningAgentsRecording.h"
 #include "LearningAgentsTrainer.h"
+#include "LearningExternalTrainer.h"
 #include "LearningAgentsTrainerEditorBase.generated.h"
 
 #define UE_API LEARNINGAGENTSTRAININGEDITOR_API
@@ -38,6 +39,15 @@ public:
 	UFUNCTION(BlueprintCallable, Category = "LearningAgents")
 	FLearningAgentsCommunicator MakeFileCommunicator(FDirectoryPath EditorIntermediateRelativePath);
 
+	/** Export training materials to file with the given export flags. No training loop is started. */
+	virtual void StartExport(UE::Learning::ETrainerExportFlags Flags) PURE_VIRTUAL(ALearningAgentsTrainerEditorBase::StartExport);
+
+	/** The export flags used by MakeFileCommunicator. */
+	UE::Learning::ETrainerExportFlags ExportFlags = UE::Learning::ETrainerExportFlags::All;
+
+	/** The file trainer created by MakeFileCommunicator, used by StartExport. */
+	TSharedPtr<UE::Learning::FFileTrainer> FileTrainer;
+
 	float TrainingTickInterval = 0.0f;
 	
 	UPROPERTY(VisibleAnywhere, BlueprintReadWrite, Category = "LearningAgents")
diff --git a/Engine/Plugins/Experimental/LearningAgents/Source/LearningAgentsTrainingEditor/Public/LearningAgentsTrainerEditorDetails.h b/Engine/Plugins/Experimental/LearningAgents/Source/LearningAgentsTrainingEditor/Public/LearningAgentsTrainerEditorDetails.h
--- a/Engine/Plugins/Experimental/LearningAgents/Source/LearningAgentsTrainingEditor/Public/LearningAgentsTrainerEditorDetails.h
+++ b/Engine/Plugins/Experimental/LearningAgents/Source/LearningAgentsTrainingEditor/Public/LearningAgentsTrainerEditorDetails.h
@@ -18,6 +18,10 @@ private:
 	TArray<TWeakObjectPtr<UObject>> EditedObjects;
 	FReply OnRunClicked();
 	FReply OnStopClicked();
+	FReply OnExportAllClicked();
+	FReply OnExportNetworksClicked();
+	FReply OnExportReplayBuffersClicked();
+	FReply OnExportTrainerConfigClicked();
 };
 
 #undef UE_API
diff --git a/Engine/Plugins/Experimental/LearningCore/Source/LearningTraining/Private/LearningExternalTrainer.cpp b/Engine/Plugins/Experimental/LearningCore/Source/LearningTraining/Private/LearningExternalTrainer.cpp
--- a/Engine/Plugins/Experimental/LearningCore/Source/LearningTraining/Private/LearningExternalTrainer.cpp
+++ b/Engine/Plugins/Experimental/LearningCore/Source/LearningTraining/Private/LearningExternalTrainer.cpp
@@ -269,6 +269,7 @@ namespace UE::Learning
 
 	ETrainerResponse FSharedMemoryTrainer::SendConfigs(
 		const TSharedRef<FJsonObject>& DataConfigObject,
+		const TSharedRef<FJsonObject>& ModelConfigObject,
 		const TSharedRef<FJsonObject>& TrainerConfigObject,
 		const ELogSetting LogSettings)
 	{
@@ -367,6 +368,12 @@ namespace UE::Learning
 		FJsonSerializer::Serialize(SharedMemoryConfigObject, SharedMemoryJsonWriter, true);
 		FFileHelper::SaveStringToFile(SharedMemoryConfigString, *(ConfigPath / FString::Printf(TEXT("shared-memory-%s.json"), *Controls.Guid.ToString(PLATFORM_MAC ? EGuidFormats::Short : EGuidFormats::DigitsWithHyphensInBraces))));
 
+		// Merge model config into data config for shared memory protocol compatibility
+		for (const auto& Pair : ModelConfigObject->Values)
+		{
+			DataConfigObject->SetField(Pair.Key, Pair.Value);
+		}
+
 		FString DataConfigString;
 		TSharedRef<TJsonWriter<>> DataJsonWriter = TJsonWriterFactory<>::Create(&DataConfigString, 0);
 		FJsonSerializer::Serialize(DataConfigObject, DataJsonWriter, true);
@@ -815,6 +822,7 @@ namespace UE::Learning
 
 	ETrainerResponse FSocketTrainer::SendConfigs(
 		const TSharedRef<FJsonObject>& DataConfigObject,
+		const TSharedRef<FJsonObject>& ModelConfigObject,
 		const TSharedRef<FJsonObject>& TrainerConfigObject,
 		const ELogSetting LogSettings)
 	{
@@ -825,6 +833,12 @@ namespace UE::Learning
 			return ETrainerResponse::Unexpected;
 		}
 
+		// Merge model config into data config for socket protocol compatibility
+		for (const auto& Pair : ModelConfigObject->Values)
+		{
+			DataConfigObject->SetField(Pair.Key, Pair.Value);
+		}
+
 		DataConfigObject->SetObjectField(TEXT("TrainerSettings"), TrainerConfigObject);
 
 		FString ConfigString;
@@ -1050,16 +1064,18 @@ namespace UE::Learning
 			LogSettings);
 	}
 
-	FFileTrainer::FFileTrainer(const FString& InIntermediatePath)
+	FFileTrainer::FFileTrainer(const FString& InIntermediatePath, ETrainerExportFlags InExportFlags)
 	{
 		IntermediatePath = InIntermediatePath;
 		FileSaveDirectory = IntermediatePath / TEXT("TrainingMaterials") / FDateTime::Now().ToString(TEXT("%Y%m%d%H%M%S"));
+		ExportFlags = InExportFlags;
 	}
 
 	FFileTrainer::~FFileTrainer(){}
 
 	ETrainerResponse FFileTrainer::SendConfigs(
 		const TSharedRef<FJsonObject>& DataConfigObject,
+		const TSharedRef<FJsonObject>& ModelConfigObject,
 		const TSharedRef<FJsonObject>& TrainerConfigObject,
 		const ELogSetting LogSettings)
 	{
@@ -1068,15 +1084,35 @@ namespace UE::Learning
 		// Add intermediate path as a temp directory for the artifacts logging on python side
 		TrainerConfigObject->SetStringField(TEXT("TempDirectory"), *FileManager.ConvertToAbsolutePathForExternalAppForRead(*IntermediatePath));
 
-		FString DataConfigString;
-		TSharedRef<TJsonWriter<>> DataJsonWriter = TJsonWriterFactory<>::Create(&DataConfigString, 0);
-		FJsonSerializer::Serialize(DataConfigObject, DataJsonWriter, true);
-		FFileHelper::SaveStringToFile(DataConfigString, *(FileSaveDirectory / "data-config.json"));
+		IPlatformFile& PlatformFile = FPlatformFileManager::Get().GetPlatformFile();
 
-		FString TrainerConfigString;
-		TSharedRef<TJsonWriter<>> TrainerJsonWriter = TJsonWriterFactory<>::Create(&TrainerConfigString, 0);
-		FJsonSerializer::Serialize(TrainerConfigObject, TrainerJsonWriter, true);
-		FFileHelper::SaveStringToFile(TrainerConfigString, *(FileSaveDirectory / "trainer-config.json"));
+		if (EnumHasAnyFlags(ExportFlags, ETrainerExportFlags::ReplayBuffers))
+		{
+			const FString ReplayBuffersDir = FileSaveDirectory / TEXT("ReplayBuffers");
+			PlatformFile.CreateDirectoryTree(*ReplayBuffersDir);
+			FString DataConfigString;
+			TSharedRef<TJsonWriter<>> DataJsonWriter = TJsonWriterFactory<>::Create(&DataConfigString, 0);
+			FJsonSerializer::Serialize(DataConfigObject, DataJsonWriter, true);
+			FFileHelper::SaveStringToFile(DataConfigString, *(ReplayBuffersDir / "data-config.json"));
+		}
+
+		if (EnumHasAnyFlags(ExportFlags, ETrainerExportFlags::Networks))
+		{
+			const FString NetworksDir = FileSaveDirectory / TEXT("Networks");
+			PlatformFile.CreateDirectoryTree(*NetworksDir);
+			FString ModelConfigString;
+			TSharedRef<TJsonWriter<>> ModelJsonWriter = TJsonWriterFactory<>::Create(&ModelConfigString, 0);
+			FJsonSerializer::Serialize(ModelConfigObject, ModelJsonWriter, true);
+			FFileHelper::SaveStringToFile(ModelConfigString, *(NetworksDir / "model-config.json"));
+		}
+
+		if (EnumHasAnyFlags(ExportFlags, ETrainerExportFlags::TrainerConfig))
+		{
+			FString TrainerConfigString;
+			TSharedRef<TJsonWriter<>> TrainerJsonWriter = TJsonWriterFactory<>::Create(&TrainerConfigString, 0);
+			FJsonSerializer::Serialize(TrainerConfigObject, TrainerJsonWriter, true);
+			FFileHelper::SaveStringToFile(TrainerConfigString, *(FileSaveDirectory / "trainer-config.json"));
+		}
 
 		UE_LOG(LogLearning, Display, TEXT("FileTrainer: wrote config files to %s."), *IntermediatePath);
 
@@ -1096,6 +1132,11 @@ namespace UE::Learning
 		FRWLock* NetworkLock,
 		const ELogSetting LogSettings)
 	{
+		if (!EnumHasAnyFlags(ExportFlags, ETrainerExportFlags::Networks))
+		{
+			return ETrainerResponse::Success;
+		}
+
 		if (!ensureMsgf(NetworksCount > NetworkId, TEXT("Network %d has not been added. Call AddNetwork prior to SendNetwork."), NetworkId))
 		{
 			return ETrainerResponse::Unexpected;
@@ -1147,6 +1188,11 @@ namespace UE::Learning
 
 	ETrainerResponse FFileTrainer::SendReplayBuffer(const int32 ReplayBufferId, const FReplayBuffer& ReplayBuffer, const ELogSetting LogSettings)
 	{
+		if (!EnumHasAnyFlags(ExportFlags, ETrainerExportFlags::ReplayBuffers))
+		{
+			return ETrainerResponse::Success;
+		}
+
 		if (ReplayBufferId < 0)
 		{
 			return ETrainerResponse::Unexpected;
diff --git a/Engine/Plugins/Experimental/LearningCore/Source/LearningTraining/Private/LearningPPOTrainer.cpp b/Engine/Plugins/Experimental/LearningCore/Source/LearningTraining/Private/LearningPPOTrainer.cpp
--- a/Engine/Plugins/Experimental/LearningCore/Source/LearningTraining/Private/LearningPPOTrainer.cpp
+++ b/Engine/Plugins/Experimental/LearningCore/Source/LearningTraining/Private/LearningPPOTrainer.cpp
@@ -176,7 +176,50 @@
 		const int32 ObservationSchemaId = 0;
 		const int32 ActionSchemaId = 0;
 
-		// Add Neural Network Config Entries
+		// Add Replay Buffers Config Entries
+		TArray<TSharedPtr<FJsonValue>> ReplayBufferObjects;
+		TSharedRef<FJsonValueObject> ReplayBufferJsonValue = MakeShared<FJsonValueObject>(ReplayBuffer.AsJsonConfig(ReplayBufferId));
+		ReplayBufferObjects.Add(ReplayBufferJsonValue);
+		DataConfigObject->SetArrayField(TEXT("ReplayBuffers"), ReplayBufferObjects);
+
+		// Schemas
+		TSharedPtr<FJsonObject> SchemasObject = MakeShared<FJsonObject>();
+
+		// Add the observation schemas
+		TArray<TSharedPtr<FJsonValue>> ObservationSchemaObjects;
+		{
+			// For this PPO trainer, add the one observation schema we have
+			TSharedPtr<FJsonObject> ObservationSchemaObject = MakeShared<FJsonObject>();
+			ObservationSchemaObject->SetNumberField(TEXT("Id"), ObservationSchemaId);
+			ObservationSchemaObject->SetStringField(TEXT("Name"), "Default");
+			ObservationSchemaObject->SetObjectField(TEXT("Schema"),
+				UE::Learning::Trainer::ConvertObservationSchemaToJSON(ObservationSchema, ObservationSchemaElement));
+
+			TSharedRef<FJsonValueObject> JsonValue = MakeShared<FJsonValueObject>(ObservationSchemaObject);
+			ObservationSchemaObjects.Add(JsonValue);
+		}
+		SchemasObject->SetArrayField(TEXT("Observations"), ObservationSchemaObjects);
+
+		// Add the action schemas
+		TArray<TSharedPtr<FJsonValue>> ActionSchemaObjects;
+		{
+			// For this PPO trainer, add the one action schema we have
+			TSharedPtr<FJsonObject> ActionSchemaObject = MakeShared<FJsonObject>();
+			ActionSchemaObject->SetNumberField(TEXT("Id"), ActionSchemaId);
+			ActionSchemaObject->SetStringField(TEXT("Name"), "Default");
+			ActionSchemaObject->SetObjectField(TEXT("Schema"),
+				UE::Learning::Trainer::ConvertActionSchemaToJSON(ActionSchema, ActionSchemaElement));
+
+			TSharedRef<FJsonValueObject> JsonValue = MakeShared<FJsonValueObject>(ActionSchemaObject);
+			ActionSchemaObjects.Add(JsonValue);
+		}
+		SchemasObject->SetArrayField(TEXT("Actions"), ActionSchemaObjects);
+
+		DataConfigObject->SetObjectField(TEXT("Schemas"), SchemasObject);
+
+		// Write Model Config
+		TSharedRef<FJsonObject> ModelConfigObject = MakeShared<FJsonObject>();
+
 		TArray<TSharedPtr<FJsonValue>> NetworkObjects;
 
 		// Policy
@@ -226,48 +269,7 @@
 			NetworkObjects.Add(JsonValue);
 		}
 
-		DataConfigObject->SetArrayField(TEXT("Networks"), NetworkObjects);
+		ModelConfigObject->SetArrayField(TEXT("Networks"), NetworkObjects);
 
-		// Add Replay Buffers Config Entries
-		TArray<TSharedPtr<FJsonValue>> ReplayBufferObjects;
-		TSharedRef<FJsonValueObject> ReplayBufferJsonValue = MakeShared<FJsonValueObject>(ReplayBuffer.AsJsonConfig(ReplayBufferId));
-		ReplayBufferObjects.Add(ReplayBufferJsonValue);
-		DataConfigObject->SetArrayField(TEXT("ReplayBuffers"), ReplayBufferObjects);
-
-		// Schemas
-		TSharedPtr<FJsonObject> SchemasObject = MakeShared<FJsonObject>();
-
-		// Add the observation schemas
-		TArray<TSharedPtr<FJsonValue>> ObservationSchemaObjects;
-		{
-			// For this PPO trainer, add the one observation schema we have
-			TSharedPtr<FJsonObject> ObservationSchemaObject = MakeShared<FJsonObject>();
-			ObservationSchemaObject->SetNumberField(TEXT("Id"), ObservationSchemaId);
-			ObservationSchemaObject->SetStringField(TEXT("Name"), "Default");
-			ObservationSchemaObject->SetObjectField(TEXT("Schema"),
-				UE::Learning::Trainer::ConvertObservationSchemaToJSON(ObservationSchema, ObservationSchemaElement));
-
-			TSharedRef<FJsonValueObject> JsonValue = MakeShared<FJsonValueObject>(ObservationSchemaObject);
-			ObservationSchemaObjects.Add(JsonValue);
-		}
-		SchemasObject->SetArrayField(TEXT("Observations"), ObservationSchemaObjects);
-
-		// Add the action schemas
-		TArray<TSharedPtr<FJsonValue>> ActionSchemaObjects;
-		{
-			// For this PPO trainer, add the one action schema we have
-			TSharedPtr<FJsonObject> ActionSchemaObject = MakeShared<FJsonObject>();
-			ActionSchemaObject->SetNumberField(TEXT("Id"), ActionSchemaId);
-			ActionSchemaObject->SetStringField(TEXT("Name"), "Default");
-			ActionSchemaObject->SetObjectField(TEXT("Schema"),
-				UE::Learning::Trainer::ConvertActionSchemaToJSON(ActionSchema, ActionSchemaElement));
-
-			TSharedRef<FJsonValueObject> JsonValue = MakeShared<FJsonValueObject>(ActionSchemaObject);
-			ActionSchemaObjects.Add(JsonValue);
-		}
-		SchemasObject->SetArrayField(TEXT("Actions"), ActionSchemaObjects);
-
-		DataConfigObject->SetObjectField(TEXT("Schemas"), SchemasObject);
-
 		// Add PPO Specific Config Entries
 		TSharedRef<FJsonObject> TrainingConfigObject = MakeShared<FJsonObject>();
@@ -305,7 +307,7 @@
 		TrainingConfigObject->SetBoolField(TEXT("UseTensorBoard"), TrainerSettings.bUseTensorboard);
 		TrainingConfigObject->SetBoolField(TEXT("SaveSnapshots"), TrainerSettings.bSaveSnapshots);
 
-		ExternalTrainer->SendConfigs(DataConfigObject, TrainingConfigObject, LogSettings);
+		ExternalTrainer->SendConfigs(DataConfigObject, ModelConfigObject, TrainingConfigObject, LogSettings);
 
 		Response = ExternalTrainer->SendNetwork(PolicyNetworkId, PolicyNetwork, PolicyNetworkLock);
 
diff --git a/Engine/Plugins/Experimental/LearningCore/Source/LearningTraining/Public/LearningExternalTrainer.h b/Engine/Plugins/Experimental/LearningCore/Source/LearningTraining/Public/LearningExternalTrainer.h
--- a/Engine/Plugins/Experimental/LearningCore/Source/LearningTraining/Public/LearningExternalTrainer.h
+++ b/Engine/Plugins/Experimental/LearningCore/Source/LearningTraining/Public/LearningExternalTrainer.h
@@ -51,13 +51,15 @@ namespace UE::Learning
 		/**
 		* Sends the given json configs to the trainer process.
 		*
-		* @param DataConfigObject		The config to send with meta-data
+		* @param DataConfigObject		The config to send with data-related settings (replay buffers, schemas)
+		* @param ModelConfigObject		The config to send with model-related settings (networks, architecture)
 		* @param TrainerConfigObject	The config to send with trainer-specific settings
 		* @param LogSettings			The log verbosity level
 		* @returns						Trainer response
 		*/
 		virtual ETrainerResponse SendConfigs(
 			const TSharedRef<FJsonObject>& DataConfigObject,
+			const TSharedRef<FJsonObject>& ModelConfigObject,
 			const TSharedRef<FJsonObject>& TrainerConfigObject,
 			const ELogSetting LogSettings = Trainer::DefaultLogSettings) = 0;
 
@@ -275,6 +277,7 @@ namespace UE::Learning
 
 		UE_API virtual ETrainerResponse SendConfigs(
 			const TSharedRef<FJsonObject>& DataConfigObject,
+			const TSharedRef<FJsonObject>& ModelConfigObject,
 			const TSharedRef<FJsonObject>& TrainerConfigObject,
 			const ELogSetting LogSettings = Trainer::DefaultLogSettings) override final;
 
@@ -425,6 +428,7 @@ namespace UE::Learning
 
 		UE_API virtual ETrainerResponse SendConfigs(
 			const TSharedRef<FJsonObject>& DataConfigObject,
+			const TSharedRef<FJsonObject>& ModelConfigObject,
 			const TSharedRef<FJsonObject>& TrainerConfigObject,
 			const ELogSetting LogSettings = Trainer::DefaultLogSettings) override final;
 
@@ -483,8 +487,9 @@ namespace UE::Learning
 		* Creates a new trainer that dumps training materials to file.
 		*
 		* @param InIntermediatePath				The Intermediate Path - for writing temporary files.
+		* @param InExportFlags					Which sections to export (defaults to All).
 		*/
-		UE_API FFileTrainer(const FString& InIntermediatePath = FString());
+		UE_API FFileTrainer(const FString& InIntermediatePath = FString(), ETrainerExportFlags InExportFlags = ETrainerExportFlags::All);
 
 		UE_API ~FFileTrainer();
 
@@ -500,6 +505,7 @@ namespace UE::Learning
 
 		UE_API virtual ETrainerResponse SendConfigs(
 			const TSharedRef<FJsonObject>& DataConfigObject,
+			const TSharedRef<FJsonObject>& ModelConfigObject,
 			const TSharedRef<FJsonObject>& TrainerConfigObject,
 			const ELogSetting LogSettings = Trainer::DefaultLogSettings) override final;
 
@@ -533,6 +539,7 @@ namespace UE::Learning
 	private:
 		FString IntermediatePath;
 		FString FileSaveDirectory;
+		ETrainerExportFlags ExportFlags = ETrainerExportFlags::All;
 		int32 ReplayBufferCount = 0;
 		int32 NetworksCount = 0;
 
diff --git a/Engine/Plugins/Experimental/LearningCore/Source/LearningTraining/Public/LearningTrainer.h b/Engine/Plugins/Experimental/LearningCore/Source/LearningTraining/Public/LearningTrainer.h
--- a/Engine/Plugins/Experimental/LearningCore/Source/LearningTraining/Public/LearningTrainer.h
+++ b/Engine/Plugins/Experimental/LearningCore/Source/LearningTraining/Public/LearningTrainer.h
@@ -74,6 +74,17 @@ namespace UE::Learning
 	};
 	ENUM_CLASS_FLAGS(ESubprocessFlags)
 
+	/** Flags controlling which sections a trainer exports. */
+	enum class ETrainerExportFlags : uint8
+	{
+		None           = 0,
+		ReplayBuffers  = 1 << 0,	// ReplayBuffers/ dir: replay binaries + data-config.json
+		Networks       = 1 << 1,	// Networks/ dir: network binaries + model-config.json
+		TrainerConfig  = 1 << 2,	// trainer-config.json
+		All            = ReplayBuffers | Networks | TrainerConfig,
+	};
+	ENUM_CLASS_FLAGS(ETrainerExportFlags)
+
 	/**
 	* Simple managed subprocess similar to FMonitoredProcess
 	*/
