[LACombat] Midrange Model: v2v3-01-continuous-o1-a1-outlierv3-batch1024

#changelist validated
#virtualized
#Synced-CL 51878233

diff --git a/Engine/Plugins/Experimental/LearningAgents/Source/LearningAgents/Private/LearningAgentsFlowMatching.cpp b/Engine/Plugins/Experimental/LearningAgents/Source/LearningAgents/Private/LearningAgentsFlowMatching.cpp
--- a/Engine/Plugins/Experimental/LearningAgents/Source/LearningAgents/Private/LearningAgentsFlowMatching.cpp
+++ b/Engine/Plugins/Experimental/LearningAgents/Source/LearningAgents/Private/LearningAgentsFlowMatching.cpp
@@ -32,6 +32,7 @@ ULearningAgentsFlowMatching* ULearningAgentsFlowMatching::MakeFlowMatchingModel(
 	const bool bReinitializeEncoderNetwork,
 	const bool bReinitializeDenoiserNetwork,
 	const bool bReinitializeDecoderNetwork,
+	const bool bSkipCompatibilityChecks,
 	const FLearningAgentsFlowMatchingSettings& FlowMatchingSettings,
 	const int32 Seed)
 {
@@ -61,6 +62,7 @@ ULearningAgentsFlowMatching* ULearningAgentsFlowMatching::MakeFlowMatchingModel(
 		bReinitializeEncoderNetwork,
 		bReinitializeDenoiserNetwork,
 		bReinitializeDecoderNetwork,
+		bSkipCompatibilityChecks,
 		FlowMatchingSettings,
 		Seed);
 
@@ -76,6 +78,7 @@ void ULearningAgentsFlowMatching::SetupFlowMatchingModel(
 	const bool bReinitializeEncoderNetwork,
 	const bool bReinitializeDenoiserNetwork,
 	const bool bReinitializeDecoderNetwork,
+	const bool bSkipCompatibilityChecks,
 	const FLearningAgentsFlowMatchingSettings& InFlowMatchingSettings,
 	const int32 Seed)
 {
@@ -135,7 +138,7 @@ void ULearningAgentsFlowMatching::SetupFlowMatchingModel(
 	{
 		EncoderNetwork = EncoderNeuralNetworkAsset;
 
-		if (EncoderNeuralNetworkAsset->NeuralNetworkData && !bReinitializeEncoderNetwork)
+		if (EncoderNeuralNetworkAsset->NeuralNetworkData && !bReinitializeEncoderNetwork && !bSkipCompatibilityChecks)
 		{
 			if (EncoderNeuralNetworkAsset->NeuralNetworkData->GetCompatibilityHash() != ObservationCompatibilityHash)
 			{
@@ -177,7 +180,7 @@ void ULearningAgentsFlowMatching::SetupFlowMatchingModel(
 	{
 		DenoiserNetwork = DenoiserNeuralNetworkAsset;
 
-		if (DenoiserNeuralNetworkAsset->NeuralNetworkData && !bReinitializeDenoiserNetwork)
+		if (DenoiserNeuralNetworkAsset->NeuralNetworkData && !bReinitializeDenoiserNetwork && !bSkipCompatibilityChecks)
 		{
 			if (DenoiserNeuralNetworkAsset->NeuralNetworkData->GetCompatibilityHash() != DenoiserCompatibilityHash)
 			{
@@ -231,7 +234,7 @@ void ULearningAgentsFlowMatching::SetupFlowMatchingModel(
 	{
 		DecoderNetwork = DecoderNeuralNetworkAsset;
 
-		if (DecoderNeuralNetworkAsset->NeuralNetworkData && !bReinitializeDecoderNetwork)
+		if (DecoderNeuralNetworkAsset->NeuralNetworkData && !bReinitializeDecoderNetwork && !bSkipCompatibilityChecks)
 		{
 			if (DecoderNeuralNetworkAsset->NeuralNetworkData->GetCompatibilityHash() != ActionCompatibilityHash)
 			{
@@ -412,10 +415,14 @@ void ULearningAgentsFlowMatching::RunInference(const int32 ODEStepsNum)
 	// ODE integration loop
 	const float DeltaTime = 1.0f / static_cast<float>(ODEStepsNum);
 
-	for (int32 Step = 0; Step < ODEStepsNum; Step++)
 	{
-		const float Time = static_cast<float>(Step) * DeltaTime;
-		PerformODEStep(ODEState, Time, DeltaTime);
+		TRACE_CPUPROFILER_EVENT_SCOPE(ULearningAgentsFlowMatching::ODEIntegration);
+
+		for (int32 Step = 0; Step < ODEStepsNum; Step++)
+		{
+			const float Time = static_cast<float>(Step) * DeltaTime;
+			PerformODEStep(ODEState, Time, DeltaTime);
+		}
 	}
 
 	// Denormalize final ODE state
@@ -564,7 +571,7 @@ void ULearningAgentsFlowMatching::MakeDenormalizerNetworkModelBuilderElementFrom
 	}
 	case UE::Learning::Action::EType::DiscreteExclusive:
 	{
-		const int32 ValueNum = Schema.GetDiscreteExclusive(SchemaElement).Num;
+		const int32 ValueNum = Schema.GetDiscreteExclusive(SchemaElement).Num * FlowMatchingSettings.ActionChunkSize;
 
 		OutElement = Builder.MakeDenormalize(
 			ValueNum,
@@ -574,7 +581,7 @@ void ULearningAgentsFlowMatching::MakeDenormalizerNetworkModelBuilderElementFrom
 	}
 	case UE::Learning::Action::EType::DiscreteInclusive:
 	{
-		const int32 ValueNum = Schema.GetDiscreteInclusive(SchemaElement).Num;
+		const int32 ValueNum = Schema.GetDiscreteInclusive(SchemaElement).Num * FlowMatchingSettings.ActionChunkSize;
 
 		OutElement = Builder.MakeDenormalize(
 			ValueNum,
@@ -584,7 +591,7 @@ void ULearningAgentsFlowMatching::MakeDenormalizerNetworkModelBuilderElementFrom
 	}
 	case UE::Learning::Action::EType::NamedDiscreteExclusive:
 	{
-		const int32 ValueNum = Schema.GetNamedDiscreteExclusive(SchemaElement).ElementNames.Num();
+		const int32 ValueNum = Schema.GetNamedDiscreteExclusive(SchemaElement).ElementNames.Num() * FlowMatchingSettings.ActionChunkSize;
 
 		OutElement = Builder.MakeDenormalize(
 			ValueNum,
@@ -594,7 +601,7 @@ void ULearningAgentsFlowMatching::MakeDenormalizerNetworkModelBuilderElementFrom
 	}
 	case UE::Learning::Action::EType::NamedDiscreteInclusive:
 	{
-		const int32 ValueNum = Schema.GetNamedDiscreteInclusive(SchemaElement).ElementNames.Num();
+		const int32 ValueNum = Schema.GetNamedDiscreteInclusive(SchemaElement).ElementNames.Num() * FlowMatchingSettings.ActionChunkSize;
 
 		OutElement = Builder.MakeDenormalize(
 			ValueNum,
diff --git a/Engine/Plugins/Experimental/LearningAgents/Source/LearningAgents/Public/LearningAgentsFlowMatching.h b/Engine/Plugins/Experimental/LearningAgents/Source/LearningAgents/Public/LearningAgentsFlowMatching.h
--- a/Engine/Plugins/Experimental/LearningAgents/Source/LearningAgents/Public/LearningAgentsFlowMatching.h
+++ b/Engine/Plugins/Experimental/LearningAgents/Source/LearningAgents/Public/LearningAgentsFlowMatching.h
@@ -74,6 +74,7 @@ public:
 		const bool bReinitializeEncoderNetwork = true,
 		const bool bReinitializeDenoiserNetwork = true,
 		const bool bReinitializeDecoderNetwork = true,
+		const bool bSkipCompatibilityChecks = false,
 		const FLearningAgentsFlowMatchingSettings& FlowMatchingSettings = FLearningAgentsFlowMatchingSettings(),
 		const int32 Seed = 1234);
 
@@ -87,6 +88,7 @@ public:
 		const bool bReinitializeEncoderNetwork = true,
 		const bool bReinitializeDenoiserNetwork = true,
 		const bool bReinitializeDecoderNetwork = true,
+		const bool bSkipCompatibilityChecks = false,
 		const FLearningAgentsFlowMatchingSettings& InFlowMatchingSettings = FLearningAgentsFlowMatchingSettings(),
 		const int32 Seed = 1234);
 
