[LACombat] Replay Runtime Recording

- LAReplay: serializations, replay data storage
- LAReplayRecorder: generic base recorder
- LAReplayTypes: data structures

- LACombatReplayRecorder: Combat specific recorder

#changelist validated
#virtualized
#Synced-CL 51878233

diff --git a/Engine/Plugins/Experimental/LearningAgents/Source/LearningAgentsTraining/Private/LearningAgentsRecorder.cpp b/Engine/Plugins/Experimental/LearningAgents/Source/LearningAgentsTraining/Private/LearningAgentsRecorder.cpp
--- a/Engine/Plugins/Experimental/LearningAgents/Source/LearningAgentsTraining/Private/LearningAgentsRecorder.cpp
+++ b/Engine/Plugins/Experimental/LearningAgents/Source/LearningAgentsTraining/Private/LearningAgentsRecorder.cpp
@@ -426,7 +426,7 @@ ULearningAgentsInteractor* ULearningAgentsRecorder::GetInteractor() const
 	return Interactor;
 }
 
-void ULearningAgentsRecorder::SaveSchemaConfig()
+void ULearningAgentsRecorder::SaveSchemaConfig(const FString& FileName)
 {
 	TSharedRef<FJsonObject> ConfigObject = MakeShared<FJsonObject>();
 	ConfigObject->SetObjectField(TEXT("ObservationSchema"),
@@ -437,5 +437,6 @@ void ULearningAgentsRecorder::SaveSchemaConfig()
 	FString ConfigString;
 	TSharedRef<TJsonWriter<>> JsonWriter = TJsonWriterFactory<>::Create(&ConfigString, 0);
 	FJsonSerializer::Serialize(ConfigObject, JsonWriter, true);
-	FFileHelper::SaveStringToFile(ConfigString, *(RecordingDirectory / "schema-config.json"));
+
+	FFileHelper::SaveStringToFile(ConfigString, *(RecordingDirectory / FileName));
 }
diff --git a/Engine/Plugins/Experimental/LearningAgents/Source/LearningAgentsTraining/Public/LearningAgentsRecorder.h b/Engine/Plugins/Experimental/LearningAgents/Source/LearningAgentsTraining/Public/LearningAgentsRecorder.h
--- a/Engine/Plugins/Experimental/LearningAgents/Source/LearningAgentsTraining/Public/LearningAgentsRecorder.h
+++ b/Engine/Plugins/Experimental/LearningAgents/Source/LearningAgentsTraining/Public/LearningAgentsRecorder.h
@@ -147,7 +147,14 @@ protected:
 	UFUNCTION(BlueprintPure, Category="LearningAgents")
 	UE_API ULearningAgentsInteractor* GetInteractor() const;
 
-// ----- Private Data ----- 
+	/**
+	 * Saves the schema config JSON to the recording directory with the given filename.
+	 * Defaults to "schema-config.json".
+	 */
+	UFUNCTION(BlueprintCallable, Category = "LearningAgents")
+	UE_API void SaveSchemaConfig(const FString& FileName = TEXT("schema-config.json"));
+
+// ----- Private Data -----
 private:
 
 	/** The agent interactor this recorder is associated with. */
@@ -162,11 +169,6 @@ private:
 	UPROPERTY(VisibleAnywhere, Transient, Category = "LearningAgents")
 	bool bIsRecording = false;
 
-// ----- Private Functions -----
-private:
-	
-	UE_API void SaveSchemaConfig();
-
 // ----- Private Data ----- 
 private:
 
